Fixing Mac build error

Missed brackets from end of function call

diff --git a/Engine/Source/Runtime/Mac/CoreAudio/Private/CoreAudioSource.cpp b/Engine/Source/Runtime/Mac/CoreAudio/Private/CoreAudioSource.cpp
--- a/Engine/Source/Runtime/Mac/CoreAudio/Private/CoreAudioSource.cpp
+++ b/Engine/Source/Runtime/Mac/CoreAudio/Private/CoreAudioSource.cpp
@@ -626,7 +626,7 @@ bool FCoreAudioSoundSource::AttachToAUGraph()
 		ErrorStatus = AudioConverterNew( &Buffer->PCMFormat, &AudioDevice->MatrixMixerInputFormat, &CoreAudioConverter );
 
 		bool bIs6ChannelOGG = Buffer->NumChannels == 6
-							&& ((Buffer->DecompressionState && Buffer->DecompressionState->UsesVorbisChannelOrdering)
+							&& ((Buffer->DecompressionState && Buffer->DecompressionState->UsesVorbisChannelOrdering())
 							|| WaveInstance->WaveData->bDecompressedFromOgg);
 
 		AudioDevice->SetupMatrixMixerInput( DestInputNumber, bIs6ChannelOGG );
