Initialize InverseTransform in FCoreAudioDevice::InitializeHardware()

diff --git a/Engine/Source/Runtime/Mac/CoreAudio/Private/CoreAudioDevice.cpp b/Engine/Source/Runtime/Mac/CoreAudio/Private/CoreAudioDevice.cpp
--- a/Engine/Source/Runtime/Mac/CoreAudio/Private/CoreAudioDevice.cpp
+++ b/Engine/Source/Runtime/Mac/CoreAudio/Private/CoreAudioDevice.cpp
@@ -48,6 +48,8 @@ bool FCoreAudioDevice::InitializeHardware()
 	// Load ogg and vorbis dlls if they haven't been loaded yet
 	LoadVorbisLibraries();
 
+	InverseTransform = FMatrix::Identity;
+
 	for( SInt32 Index = 0; Index < MAX_AUDIOCHANNELS; ++Index )
 	{
 		Mixer3DInputStatus[ Index ] = false;
