You can only register the radio component once on OS X, so try to find it first & then only register if we fail to find it.

diff --git a/Engine/Source/Runtime/Mac/CoreAudio/Private/CoreAudioEffects.cpp b/Engine/Source/Runtime/Mac/CoreAudio/Private/CoreAudioEffects.cpp
--- a/Engine/Source/Runtime/Mac/CoreAudio/Private/CoreAudioEffects.cpp
+++ b/Engine/Source/Runtime/Mac/CoreAudio/Private/CoreAudioEffects.cpp
@@ -49,8 +49,13 @@ static CFBundleRef LoadRadioEffectComponent()
 							theDescription.componentSubType = 'Rdio';
 							theDescription.componentManufacturer = 'Epic';
 							theDescription.componentFlagsMask = 0;
-							AudioComponent RadioComponent = AudioComponentRegister(&theDescription, CFSTR("Epic Games: RadioEffectUnit"), Version, Factory);
-							check(RadioComponent);
+							
+							AudioComponent RadioComponent = AudioComponentFindNext(nullptr, &theDescription);
+							if ( !RadioComponent )
+							{
+								RadioComponent = AudioComponentRegister(&theDescription, CFSTR("Epic Games: RadioEffectUnit"), Version, Factory);
+								check(RadioComponent);
+							}
 							bLoaded = (RadioComponent != NULL);
 						}
 					}
