UnrealAudio -  Checkin CoreAudio implementation

- checked in first-pass on CoreAudio (Mac/iOS) implementation of audio device module of new UnrealAudio
- fixed an issue with wrapping phase for sine functions (and thus fixed unit tests)
- changed simple device out unit test to be more simple (pure sine tone on each channel in octaves every second)

diff --git a/Engine/Source/Editor/UnrealEd/UnrealEd.Build.cs b/Engine/Source/Editor/UnrealEd/UnrealEd.Build.cs
--- a/Engine/Source/Editor/UnrealEd/UnrealEd.Build.cs
+++ b/Engine/Source/Editor/UnrealEd/UnrealEd.Build.cs
@@ -233,7 +233,7 @@ public class UnrealEd : ModuleRules
 			(Target.Platform == UnrealTargetPlatform.Win32))
 		{
 			PublicDependencyModuleNames.Add("XAudio2");
-			PublicDependencyModuleNames.Add("UnrealAudioWasapi");
+			PublicDependencyModuleNames.Add("UnrealAudioXAudio2");
 
 			AddThirdPartyPrivateStaticDependencies(Target, 
 				"UEOgg",
@@ -245,6 +245,11 @@ public class UnrealEd : ModuleRules
 			
 		}
 
+		if (Target.Platform == UnrealTargetPlatform.Mac)
+		{
+			PublicDependencyModuleNames.Add("UnrealAudioCoreAudio");
+		}
+
         if (Target.Platform == UnrealTargetPlatform.HTML5)
         {
             PublicDependencyModuleNames.Add("HTML5Audio");
diff --git a/Engine/Source/Runtime/Mac/UnrealAudioCoreAudio/Private/UnrealAudioCoreAudio.cpp b/Engine/Source/Runtime/Mac/UnrealAudioCoreAudio/Private/UnrealAudioCoreAudio.cpp
new file mode 100644
--- /dev/null
+++ b/Engine/Source/Runtime/Mac/UnrealAudioCoreAudio/Private/UnrealAudioCoreAudio.cpp
@@ -0,0 +1,1333 @@
+// Copyright 1998-2015 Epic Games, Inc. All Rights Reserved.
+
+/**
+Implementation of IAudioDevice for CoreAudio
+ 
+See https://developer.apple.com/library/mac/documentation/MusicAudio/Conceptual/CoreAudioOverview/Introduction/Introduction.html
+ 
+*/
+
+#include "UnrealAudioDeviceModule.h"
+#include "ModuleManager.h"
+
+/**
+* CoreAudio System Headers
+*/
+#include <CoreAudio/AudioHardware.h>
+#include <AudioToolbox/AudioToolbox.h>
+#include <AudioUnit/AudioUnit.h>
+
+#if ENABLE_UNREAL_AUDIO
+
+// Used to toggle on white noise testing for CoreAudio output
+#define UNREAL_AUDIO_TEST_WHITE_NOISE 0
+
+DEFINE_LOG_CATEGORY(LogUnrealAudioDevice);
+
+static const TCHAR* GetCoreAudioError(OSStatus Result)
+{
+	switch (Result)
+	{
+		case kAudioHardwareNotRunningError:				return TEXT("kAudioHardwareNotRunningError");
+		case kAudioHardwareUnspecifiedError:			return TEXT("kAudioHardwareUnspecifiedError");
+		case kAudioHardwareUnknownPropertyError:		return TEXT("kAudioHardwareUnknownPropertyError");
+		case kAudioHardwareBadPropertySizeError:		return TEXT("kAudioHardwareBadPropertySizeError");
+		case kAudioHardwareIllegalOperationError:		return TEXT("kAudioHardwareIllegalOperationError");
+		case kAudioHardwareBadObjectError:				return TEXT("kAudioHardwareBadObjectError");
+		case kAudioHardwareBadDeviceError:				return TEXT("kAudioHardwareBadDeviceError");
+		case kAudioHardwareBadStreamError:				return TEXT("kAudioHardwareBadStreamError");
+		case kAudioHardwareUnsupportedOperationError:	return TEXT("kAudioHardwareUnsupportedOperationError");
+		case kAudioDeviceUnsupportedFormatError:		return TEXT("kAudioDeviceUnsupportedFormatError");
+		case kAudioDevicePermissionsError:				return TEXT("kAudioDevicePermissionsError");
+		default:										return TEXT("Unknown CoreAudio Error");
+	}
+}
+
+// Helper macro to report CoreAudio API errors
+#define CORE_AUDIO_ERR(Status, Context)													\
+	if (Status != noErr)														\
+	{																			\
+		const TCHAR* Err = GetCoreAudioError(Status);							\
+		FString Message = FString::Printf(TEXT("%s: %s"), TEXT(Context), Err);	\
+		UA_DEVICE_PLATFORM_ERROR(*Message);										\
+		return false;															\
+	}
+
+// Helper macro to make a new global audio property address
+#define NEW_GLOBAL_PROPERTY(SELECTOR)		\
+	{										\
+		SELECTOR,							\
+		kAudioObjectPropertyScopeGlobal,	\
+		kAudioObjectPropertyElementMaster	\
+	}
+
+// Helper macro to make a new output audio property address
+#define NEW_OUTPUT_PROPERTY(SELECTOR)		\
+	{										\
+		SELECTOR,							\
+		kAudioDevicePropertyScopeOutput,	\
+		kAudioObjectPropertyElementMaster	\
+	}
+
+namespace UAudio
+{
+	class FUnrealAudioCoreAudio :   public IUnrealAudioDeviceModule,
+									public FRunnable
+	{
+	public:
+		FUnrealAudioCoreAudio();
+		~FUnrealAudioCoreAudio();
+		
+		// IUnrealAudioDeviceModule
+		virtual bool Initialize() override;
+		virtual bool Shutdown() override;
+		virtual bool GetDevicePlatformApi(EDeviceApi::Type & OutType) const override;
+		virtual bool GetNumOutputDevices(uint32& OutNumDevices) const override;
+		virtual bool GetOutputDeviceInfo(const uint32 DeviceIndex, FDeviceInfo& OutInfo) const override;
+		virtual bool GetDefaultOutputDeviceIndex(uint32& OutDefaultIndex) const override;
+		virtual bool StartStream() override;
+		virtual bool StopStream() override;
+		virtual bool ShutdownStream() override;
+		virtual bool GetLatency(uint32& OutputDeviceLatency) const override;
+		virtual bool GetFrameRate(uint32& OutFrameRate) const override;
+		virtual bool OpenDevice(const FCreateStreamParams& Params) override;
+		
+		bool PerformCallback(AudioDeviceID DeviceId, const AudioBufferList* OutputBufferData);
+		void SetOverloaded();
+
+		// FRunnable, used to shutdown CoreAudio callback
+		uint32 Run() override;
+
+	private:
+
+		// Helper Functions
+		bool InitRunLoop();
+
+		bool GetOutputDevices();
+		bool GetDefaultOutputDevice();
+		bool GetDeviceInfos();
+		bool GetDeviceInfo(AudioDeviceID DeviceID, FDeviceInfo& DeviceInfo);
+		bool GetDeviceName(AudioDeviceID DeviceID, FString& OutName);
+		bool GetDeviceChannels(AudioDeviceID DeviceID, TArray<ESpeaker::Type>& OutChannels);
+		bool GetDeviceChannelsForStereo(AudioDeviceID DeviceID, TArray<ESpeaker::Type>& OutChannels);
+		bool GetDeviceChannelsForLayoutDescriptions(AudioChannelLayout* ChannelLayout, TArray<ESpeaker::Type>& OutChannels);
+		bool GetDeviceChannelsForBitMap(UInt32 BitMap, TArray<ESpeaker::Type>& OutChannels);
+		bool GetDeviceChannelsForLayoutTag(AudioChannelLayoutTag LayoutTag, TArray<ESpeaker::Type>& OutChannels);
+		bool GetDeviceChannelsForChannelCount(UInt32 NumChannels, TArray<ESpeaker::Type>& OutChannels);
+		bool GetDeviceFrameRates(AudioDeviceID DeviceID, TArray<uint32>& PossibleFramerates, uint32& DefaultFrameRate);
+		bool GetDeviceLatency(AudioDeviceID DeviceID, UInt32& Latency);
+
+		bool InitDeviceOutputId(UInt32 DeviceIndex);
+		bool InitDeviceCallback(const FCreateStreamParams& Params);
+		bool InitDeviceFrameRate(const UInt32 RequestedFrameRate, bool& bChanged);
+		bool InitDeviceVirtualFormat(bool bSampleRateChanged);
+		bool InitDevicePhysicalFormat();
+		bool InitDeviceNumDeviceStreams();
+		bool InitDeviceOverrunCallback();
+
+		/** CoreAudio-specific state */
+		struct FCoreAudioInfo
+		{
+			FCoreAudioInfo()
+				: OutputDeviceId(INDEX_NONE)
+				, OutputDeviceIndex(INDEX_NONE)
+				, DefaultDeviceId(INDEX_NONE)
+				, DefaultDeviceIndex(INDEX_NONE)
+				, DeviceIOProcId(nullptr)
+				, NumDeviceStreams(1)
+				, StoppingCallbackCount(0)
+				, bInternalStop(false)
+			{
+			}
+
+			/** The open output device id */
+			AudioDeviceID OutputDeviceId;
+
+			/** The open output device index (into array of output devices) */
+			UInt32 OutputDeviceIndex;
+
+			/** The default device Id */
+			AudioDeviceID DefaultDeviceId;
+
+			/** The index into the device arrays which corresponds to the default output device. */
+			int32 DefaultDeviceIndex;
+
+			/** The device IO process Id */
+			AudioDeviceIOProcID DeviceIOProcId;
+
+			/** Array of output device info. This is built on initialize. */
+			TArray<FDeviceInfo> OutputDevices;
+
+			/** Array of CoreAudio DeviceID objects */
+			TArray<AudioDeviceID> OutputDeviceIDs;
+
+			/** 
+			* Number of device streams the device has. Some devices
+			* can have multiple output streams for multi-channel playback. 
+			*/
+			UInt32 NumDeviceStreams;
+
+			/** Info about the callback passed to user */
+			FCallbackInfo CallbackInfo;
+
+			/** Condition used for signalling callback thread during shutdown */
+			pthread_cond_t Condition;
+
+			/** Mutex used to block stop thread during CoreAudio callback shutdown */
+			pthread_mutex_t Mutex;
+
+			/** Number used to track shutdown states for CoreAudio callback */
+			int32 StoppingCallbackCount;
+
+			/** Whether or not a stop was initiated from inside the CoreAudio callback */
+			bool bInternalStop;
+		};
+
+		/** CoreAudio-specific state */
+		FCoreAudioInfo CoreAudioInfo;
+
+		/** Whether or not we've initialized the device */
+		bool bInitialized;
+	};
+
+	/**
+	* CoreAudio Callbacks
+	*/
+
+	/** A CoreAudio Property Listener to listen to the sample rate property when it gets changed. */
+	static OSStatus SampleRatePropertyListener(	AudioObjectID DeviceId,
+												UInt32 NumAddresses,
+												const AudioObjectPropertyAddress* Addresses,
+												void* SampleRatePtr)
+	{
+		Float64* SampleRate = (Float64*)SampleRatePtr;
+		UInt32 DataSize = sizeof(Float64);
+		AudioObjectPropertyAddress PropertyAddress = NEW_GLOBAL_PROPERTY(kAudioDevicePropertyNominalSampleRate);
+		OSStatus Status = AudioObjectGetPropertyData(DeviceId, &PropertyAddress, 0, nullptr, &DataSize, SampleRate);
+		return Status;
+	}
+
+	/** A CoreAudio Property Listener to listen to whether or not the audio callback is overloaded (resulting in hitching or buffer over/underun). */
+	static OSStatus OverrunPropertyListener(AudioObjectID DeviceId,
+											UInt32 NumAddresses,
+											const AudioObjectPropertyAddress* Addresses,
+											void* UserData)
+	{
+		for (UInt32 i = 0; i < NumAddresses; ++i)
+		{
+			if (Addresses[i].mSelector == kAudioDeviceProcessorOverload)
+			{
+				class UAudio::FUnrealAudioCoreAudio* UnrealAudioCoreAudio = (class UAudio::FUnrealAudioCoreAudio*)UserData;
+				UnrealAudioCoreAudio->SetOverloaded();
+			}
+		}
+		return kAudioHardwareNoError;
+	}
+
+	/** Core audio callback function called when the output device is ready for more data. */
+	static OSStatus CoreAudioCallback(	AudioDeviceID DeviceId,
+										const AudioTimeStamp* CurrentTimeStamp,
+										const AudioBufferList* InputBufferData,
+										const AudioTimeStamp* InputTime,
+										AudioBufferList* OutputBufferData,
+										const AudioTimeStamp* OutputTime,
+										void* UserData)
+	{
+		// Get the user data and cast to our FUnrealAudioCoreAudio object
+		FUnrealAudioCoreAudio* UnrealAudioCoreAudio = (FUnrealAudioCoreAudio*) UserData;
+
+		// Call our callback function
+		if (!UnrealAudioCoreAudio->PerformCallback(DeviceId, OutputBufferData))
+		{
+			// Something went wrong...
+			return kAudioHardwareUnspecifiedError;
+		}
+		// Everything went cool...
+		return kAudioHardwareNoError;
+	}
+
+	/**
+	* FUnrealAudioCoreAudio Implementation
+	*/
+
+	FUnrealAudioCoreAudio::FUnrealAudioCoreAudio()
+		: bInitialized(false)
+	{
+	}
+
+	FUnrealAudioCoreAudio::~FUnrealAudioCoreAudio()
+	{
+		if (bInitialized)
+		{
+			Shutdown();
+		}
+	}
+
+	bool FUnrealAudioCoreAudio::Initialize()
+	{
+		if (bInitialized)
+		{
+			return false;
+		}
+		bool bSuccess = InitRunLoop();
+		bSuccess &= GetOutputDevices();
+		bSuccess &= GetDefaultOutputDevice();
+		bSuccess &= GetDeviceInfos();
+		return (bInitialized = bSuccess);
+	}
+
+	bool FUnrealAudioCoreAudio::Shutdown()
+	{
+		if (StreamInfo.State != EStreamState::SHUTDOWN)
+		{
+			ShutdownStream();
+		}
+		return true;
+	}
+	
+	bool FUnrealAudioCoreAudio::GetDevicePlatformApi(EDeviceApi::Type & OutType) const
+	{
+		OutType = EDeviceApi::CORE_AUDIO;
+		return true;
+	}
+	
+	bool FUnrealAudioCoreAudio::GetNumOutputDevices(uint32& OutNumDevices) const
+	{
+		OutNumDevices = CoreAudioInfo.OutputDevices.Num();
+		return true;
+	}
+	
+	bool FUnrealAudioCoreAudio::GetOutputDeviceInfo(const uint32 DeviceIndex, FDeviceInfo& OutInfo) const
+	{
+		if (DeviceIndex < CoreAudioInfo.OutputDevices.Num())
+		{
+			OutInfo = CoreAudioInfo.OutputDevices[DeviceIndex];
+			return true;
+		}
+		return false;
+	}
+	
+	bool FUnrealAudioCoreAudio::GetDefaultOutputDeviceIndex(uint32& OutDefaultIndex) const
+	{
+		OutDefaultIndex = CoreAudioInfo.DefaultDeviceIndex;
+		return true;
+	}
+
+	bool FUnrealAudioCoreAudio::StartStream()
+	{
+		if (!bInitialized || StreamInfo.State == EStreamState::RUNNING)
+		{
+			return false;
+		}
+
+		// Start up the audio device stream
+		OSStatus Status = AudioDeviceStart(CoreAudioInfo.OutputDeviceId, CoreAudioCallback);
+		CORE_AUDIO_ERR(Status, "Failed to start audio device");
+
+		StreamInfo.State = EStreamState::RUNNING;
+		CoreAudioInfo.StoppingCallbackCount = 0;
+		return true;
+	}
+
+	bool FUnrealAudioCoreAudio::StopStream()
+	{
+		if (StreamInfo.State == EStreamState::STOPPED)
+		{
+			return false;
+		}
+
+		// If this was a non-internal, stop, then set the count to 2
+		// and then wait for the device callback to finish its current callback.
+		if (!CoreAudioInfo.StoppingCallbackCount)
+		{
+			CoreAudioInfo.StoppingCallbackCount = 2;
+			pthread_cond_wait(&CoreAudioInfo.Condition, &CoreAudioInfo.Mutex);
+		}
+
+		OSStatus Status = AudioDeviceStop(CoreAudioInfo.OutputDeviceId, CoreAudioCallback);
+		CORE_AUDIO_ERR(Status, "Failed to stop audio device callback.");
+
+		StreamInfo.State = EStreamState::STOPPED;
+
+		if (StreamInfo.Thread)
+		{
+			delete StreamInfo.Thread;
+			StreamInfo.Thread = nullptr;
+		}
+		return true;
+	}
+	
+	bool FUnrealAudioCoreAudio::ShutdownStream()
+	{
+		if (StreamInfo.State == EStreamState::SHUTDOWN)
+		{
+			return false;
+		}
+
+		AudioObjectPropertyAddress Property = NEW_GLOBAL_PROPERTY(kAudioDeviceProcessorOverload);
+		OSStatus Status = AudioObjectRemovePropertyListener(CoreAudioInfo.OutputDeviceId, &Property, OverrunPropertyListener, (void *)this);
+		CORE_AUDIO_ERR(Status, "Failed to remove device overrun property listener");
+
+		if (StreamInfo.State == EStreamState::RUNNING)
+		{
+			Status = AudioDeviceStop(CoreAudioInfo.OutputDeviceId, CoreAudioCallback);
+			CORE_AUDIO_ERR(Status, "Failed to stop audio device");
+		}
+
+		Status = AudioDeviceDestroyIOProcID(CoreAudioInfo.OutputDeviceId, CoreAudioInfo.DeviceIOProcId);
+		CORE_AUDIO_ERR(Status, "Failed to stop audio destroy IOProcID");
+
+		pthread_cond_destroy(&CoreAudioInfo.Condition);
+		StreamInfo.State = EStreamState::SHUTDOWN;
+		return true;
+	}
+	
+	bool FUnrealAudioCoreAudio::GetLatency(uint32& OutputDeviceLatency) const
+	{
+		OutputDeviceLatency = StreamInfo.DeviceInfo.Latency;
+		return true;
+	}
+	
+	bool FUnrealAudioCoreAudio::GetFrameRate(uint32& OutFrameRate) const
+	{
+		OutFrameRate = StreamInfo.FrameRate;
+		return true;
+	}
+	
+	bool FUnrealAudioCoreAudio::OpenDevice(const FCreateStreamParams& Params)
+	{
+		if (!bInitialized || StreamInfo.State != EStreamState::SHUTDOWN)
+		{
+			return false;
+		}
+
+		bool bSampleRateChanged = false;
+		bool bSuccess = InitDeviceOutputId(Params.OutputDeviceIndex);
+		bSuccess &= InitDeviceFrameRate(Params.FrameRate, bSampleRateChanged);
+		bSuccess &= InitDeviceVirtualFormat(bSampleRateChanged);
+		bSuccess &= InitDevicePhysicalFormat();
+		bSuccess &= InitDeviceNumDeviceStreams();
+		bSuccess &= InitDeviceCallback(Params);
+		bSuccess &= InitDeviceOverrunCallback();
+		StreamInfo.State = bSuccess ? EStreamState::STOPPED : EStreamState::SHUTDOWN;
+		return bSuccess;
+	}
+
+	uint32 FUnrealAudioCoreAudio::Run()
+	{
+		// Just call stop stream
+		StopStream();
+		return 0;
+	}
+
+	bool FUnrealAudioCoreAudio::PerformCallback(AudioDeviceID DeviceId, const AudioBufferList* OutputBuffer)
+	{
+		if (StreamInfo.State == EStreamState::STOPPED || StreamInfo.State == EStreamState::STOPPING)
+		{
+			return true;
+		}
+
+		if (StreamInfo.State == EStreamState::SHUTDOWN)
+		{
+			UA_DEVICE_PLATFORM_ERROR(TEXT("Callback called while stream was closed."));
+			return false;
+		}
+
+		if (CoreAudioInfo.StoppingCallbackCount > 3)
+		{
+			StreamInfo.State = EStreamState::STOPPING;
+			// If this shutdown was initiated internally, we need to create a thread just to call StopStream
+			// to avoid deadlocking with this callback
+			if (CoreAudioInfo.bInternalStop)
+			{
+				StreamInfo.Thread = FRunnableThread::Create(this, TEXT("CoreAudio Shutdown Thread"), 0, TPri_AboveNormal);
+			}
+			else
+			{
+				pthread_cond_signal(&CoreAudioInfo.Condition);
+			}
+			return true;
+		}
+
+		AudioDeviceID DeviceID = CoreAudioInfo.OutputDeviceId;
+		if (!CoreAudioInfo.StoppingCallbackCount)
+		{
+			FCallbackInfo& CallbackInfo = CoreAudioInfo.CallbackInfo;
+			CallbackInfo.StreamTime = StreamInfo.StreamTime;
+
+#if UNREAL_AUDIO_TEST_WHITE_NOISE
+			for (uint32 Sample = 0; Sample < CallbackInfo.NumSamples; ++Sample)
+			{
+				CallbackInfo.OutBuffer[Sample] = 0.5f * FMath::FRandRange(-1.0f, 1.0f);
+			}
+#else // #if UNREAL_AUDIO_TEST_WHITE_NOISE
+			memset((void*)CallbackInfo.OutBuffer, 0, CallbackInfo.NumSamples*sizeof(float));
+			if (!StreamInfo.CallbackFunction(CallbackInfo))
+			{
+				UA_DEVICE_PLATFORM_ERROR(TEXT("Error occurred in user callback, stopping CoreAudio thread."));
+				StreamInfo.State = EStreamState::STOPPING;
+				CoreAudioInfo.StoppingCallbackCount = 1;
+				CoreAudioInfo.bInternalStop = true;
+			}
+			// Clear status flags after the mix callback
+			CallbackInfo.StatusFlags = 0;
+#endif // #else // #if UNREAL_AUDIO_TEST_WHITE_NOISE
+		}
+
+		// If we are stopping, write zero to output buffers
+		if (CoreAudioInfo.StoppingCallbackCount)
+		{
+			for (uint32 i = 0; i < CoreAudioInfo.NumDeviceStreams; ++i)
+			{
+				memset(OutputBuffer->mBuffers[i].mData, 0, OutputBuffer->mBuffers[i].mDataByteSize);
+			}
+			++CoreAudioInfo.StoppingCallbackCount;
+		}
+		else if (CoreAudioInfo.NumDeviceStreams == 1)
+		{
+			memcpy(OutputBuffer->mBuffers[0].mData, CoreAudioInfo.CallbackInfo.OutBuffer, OutputBuffer->mBuffers[0].mDataByteSize);
+		}
+		else
+		{
+			// TODO: Support streaming to multiple output streams
+			// Report the an error (but only once)
+			static bool bPrintedError = false;
+			if (!bPrintedError)
+			{
+				bPrintedError = true;
+				UA_DEVICE_PLATFORM_ERROR(TEXT("Streaming to multiple output streams not currently supported"));
+			}
+		}
+
+		UpdateStreamTimeTick();
+		return true;
+	}
+
+	void FUnrealAudioCoreAudio::SetOverloaded()
+	{
+		UA_DEVICE_PLATFORM_ERROR(TEXT("CoreAudio Overload (Buffer Underun or Overflow) occurred."));
+		StreamInfo.CallbackInfo.StatusFlags = EStreamFlowStatus::OUTPUT_OVERFLOW;
+	}
+
+	bool FUnrealAudioCoreAudio::InitRunLoop()
+	{
+		CFRunLoopRef RunLoop = nullptr;
+		AudioObjectPropertyAddress Property = NEW_GLOBAL_PROPERTY(kAudioHardwarePropertyRunLoop);
+		OSStatus Status = AudioObjectSetPropertyData(kAudioObjectSystemObject, &Property, 0, nullptr, sizeof(CFRunLoopRef), &RunLoop);
+		CORE_AUDIO_ERR(Status, "Failed to initialize run loop");
+		return true;
+	}
+
+	bool FUnrealAudioCoreAudio::GetOutputDevices()
+	{
+		UInt32 DataSize;
+		AudioObjectPropertyAddress Property = NEW_GLOBAL_PROPERTY(kAudioHardwarePropertyDevices);
+		OSStatus Status = AudioObjectGetPropertyDataSize(kAudioObjectSystemObject, &Property, 0, nullptr, &DataSize);
+		CORE_AUDIO_ERR(Status, "Failed to get size of devices property");
+
+		UInt32 NumDevices = DataSize / sizeof(AudioDeviceID);
+		AudioDeviceID Devices[NumDevices];
+
+		Status = AudioObjectGetPropertyData(kAudioObjectSystemObject, &Property, 0, nullptr, &DataSize, (void*)&Devices);
+		CORE_AUDIO_ERR(Status, "Failed to get device list");
+
+		CoreAudioInfo.OutputDeviceIDs.Empty();
+		for (UInt32 i = 0; i < NumDevices; ++i)
+		{
+			AudioDeviceID DeviceID = Devices[i];
+			Property = NEW_OUTPUT_PROPERTY(kAudioDevicePropertyStreamConfiguration);
+			Status = AudioObjectGetPropertyDataSize(DeviceID, &Property, 0, nullptr, &DataSize);
+			CORE_AUDIO_ERR(Status, "Failed to get stream configuration size");
+
+			AudioBufferList* BufferList = (AudioBufferList*)alloca(DataSize);
+
+			Status = AudioObjectGetPropertyData(DeviceID, &Property, 0, nullptr, &DataSize, BufferList);
+			CORE_AUDIO_ERR(Status, "Failed to get stream configuration");
+
+			if (BufferList->mNumberBuffers != 0)
+			{
+				CoreAudioInfo.OutputDeviceIDs.Add(DeviceID);
+			}
+		}
+		return true;
+	}
+
+	bool FUnrealAudioCoreAudio::GetDefaultOutputDevice()
+	{
+		UInt32 PropertySize = sizeof(AudioDeviceID);
+		AudioObjectPropertyAddress Property = NEW_GLOBAL_PROPERTY(kAudioHardwarePropertyDefaultOutputDevice);
+		OSStatus Result = AudioObjectGetPropertyData(kAudioObjectSystemObject, &Property, 0, nullptr, &PropertySize, &CoreAudioInfo.DefaultDeviceId);
+		CORE_AUDIO_ERR(Result, "Failed to get default output device");
+		return true;
+	}
+
+	bool FUnrealAudioCoreAudio::GetDeviceInfos()
+	{
+		for (int32 i = 0; i < CoreAudioInfo.OutputDeviceIDs.Num(); ++i)
+		{
+			AudioDeviceID DeviceID = CoreAudioInfo.OutputDeviceIDs[i];
+			FDeviceInfo DeviceInfo;
+			if (!GetDeviceInfo(DeviceID, DeviceInfo))
+			{
+				return false;
+			}
+
+			if (DeviceID == CoreAudioInfo.DefaultDeviceId)
+			{
+				DeviceInfo.bIsSystemDefault = true;
+				CoreAudioInfo.DefaultDeviceIndex = i;
+			}
+			CoreAudioInfo.OutputDevices.Add(DeviceInfo);
+		}
+		return true;
+	}
+
+	bool FUnrealAudioCoreAudio::GetDeviceInfo(AudioDeviceID DeviceID, FDeviceInfo& DeviceInfo)
+	{
+		bool bSuccess = GetDeviceName(DeviceID, DeviceInfo.FriendlyName);
+		bSuccess &=	GetDeviceChannels(DeviceID, DeviceInfo.Speakers);
+		bSuccess &= GetDeviceFrameRates(DeviceID, DeviceInfo.PossibleFrameRates, DeviceInfo.FrameRate);
+		bSuccess &= GetDeviceLatency(DeviceID, DeviceInfo.Latency);
+		DeviceInfo.StreamFormat = EStreamFormat::FLT;
+		DeviceInfo.NumChannels = DeviceInfo.Speakers.Num();
+		return bSuccess;
+	}
+
+	bool FUnrealAudioCoreAudio::GetDeviceName(AudioDeviceID DeviceId, FString& OutName)
+	{
+		AudioObjectPropertyAddress Property = NEW_GLOBAL_PROPERTY(kAudioObjectPropertyManufacturer);
+
+		CFStringRef ManufacturerName;
+		UInt32 DataSize = sizeof(CFStringRef);
+		OSStatus Status = AudioObjectGetPropertyData(DeviceId, &Property, 0, nullptr, &DataSize, &ManufacturerName);
+		CORE_AUDIO_ERR(Status, "Failed to get device manufacturer name");
+
+		// Convert the Apple CFStringRef to an Unreal FString
+		TCHAR ManufacturerNameBuff[256];
+		FPlatformString::CFStringToTCHAR(ManufacturerName, ManufacturerNameBuff);
+		CFRelease(ManufacturerName);
+
+		CFStringRef DeviceName;
+		DataSize = sizeof(CFStringRef);
+		Property.mSelector = kAudioObjectPropertyName;
+		Status = AudioObjectGetPropertyData(DeviceId, &Property, 0, nullptr, &DataSize, &DeviceName);
+		CORE_AUDIO_ERR(Status, "Failed to get device name");
+
+		// Convert the Apple CFStringRef to an Unreal FString
+		TCHAR DeviceNameBuff[256];
+		FPlatformString::CFStringToTCHAR(DeviceName, DeviceNameBuff);
+		CFRelease(DeviceName);
+
+		OutName = FString::Printf(TEXT("%s - %s"), ManufacturerNameBuff, DeviceNameBuff);
+		return true;
+	}
+
+	bool FUnrealAudioCoreAudio::GetDeviceChannels(AudioDeviceID DeviceID, TArray<ESpeaker::Type>& OutChannels)
+	{
+		// First get number of channels by getting the buffer list
+		UInt32 PropertySize;
+		AudioObjectPropertyAddress Property = NEW_OUTPUT_PROPERTY(kAudioDevicePropertyStreamConfiguration);
+		OSStatus Status = AudioObjectGetPropertyDataSize(DeviceID, &Property, 0, nullptr, &PropertySize);
+		CORE_AUDIO_ERR(Status, "Failed to get stream configuration property size");
+
+		AudioBufferList* BufferList = (AudioBufferList*)alloca(PropertySize);
+		Status = AudioObjectGetPropertyData(DeviceID, &Property, 0, nullptr, &PropertySize, BufferList);
+		CORE_AUDIO_ERR(Status, "Failed to get stream configuration property");
+
+		UInt32 NumChannels = 0;
+		for (UInt32 i = 0; i < BufferList->mNumberBuffers; ++i)
+		{
+			NumChannels += BufferList->mBuffers[i].mNumberChannels;
+		}
+
+		if (NumChannels == 0)
+		{
+			UA_DEVICE_PLATFORM_ERROR(TEXT("Output device has 0 channels"));
+			return false;
+		}
+
+		// CoreAudio has a different property for stereo speaker layouts
+		if (NumChannels == 2)
+		{
+			return GetDeviceChannelsForStereo(DeviceID, OutChannels);
+		}
+
+		// Now CORE_AUDIO_ERR multi-channel layouts
+		Property = NEW_OUTPUT_PROPERTY(kAudioDevicePropertyPreferredChannelLayout);
+		if (!AudioObjectHasProperty(DeviceID, &Property))
+		{
+			return false;
+		}
+
+		// Bool used to try different methods of getting channel layout (which is surprisingly complex!)
+		bool bSuccess = true;
+		Status = AudioObjectGetPropertyDataSize(DeviceID, &Property, 0, nullptr, &PropertySize);
+		CORE_AUDIO_ERR(Status, "Failed to get preferred channel layout size");
+
+		AudioChannelLayout* ChannelLayout = (AudioChannelLayout*)alloca(PropertySize);
+		Status = AudioObjectGetPropertyData(DeviceID, &Property, 0, nullptr, &PropertySize, ChannelLayout);
+		CORE_AUDIO_ERR(Status, "Failed to get preferred channel layout");
+
+		if (ChannelLayout->mChannelLayoutTag == kAudioChannelLayoutTag_UseChannelDescriptions)
+		{
+			bSuccess = GetDeviceChannelsForLayoutDescriptions(ChannelLayout, OutChannels);
+		}
+		else if (ChannelLayout->mChannelLayoutTag == kAudioChannelLayoutTag_UseChannelBitmap)
+		{
+			bSuccess = GetDeviceChannelsForBitMap(ChannelLayout->mChannelBitmap, OutChannels);
+		}
+		else
+		{
+			bSuccess = GetDeviceChannelsForLayoutTag(ChannelLayout->mChannelLayoutTag, OutChannels);
+		}
+
+		// If everything failed, then lets just get the number of channels and "guess" the layout
+		if (!bSuccess)
+		{
+			bSuccess = GetDeviceChannelsForChannelCount(NumChannels, OutChannels);
+		}
+
+		return bSuccess;
+	}
+
+	bool FUnrealAudioCoreAudio::GetDeviceFrameRates(AudioDeviceID DeviceID, TArray<uint32>& PossibleFrameRates, uint32& DefaultFrameRate)
+	{
+		AudioObjectPropertyAddress Property = NEW_OUTPUT_PROPERTY(kAudioDevicePropertyAvailableNominalSampleRates);
+
+		UInt32 DataSize;
+		OSStatus Status = AudioObjectGetPropertyDataSize(DeviceID, &Property, 0, nullptr, &DataSize);
+		CORE_AUDIO_ERR(Status, "Failed to get nominal sample rates size");
+
+		UInt32 NumSampleRates = DataSize / sizeof(AudioValueRange);
+		AudioValueRange SampleRates[NumSampleRates];
+
+		Status = AudioObjectGetPropertyData(DeviceID, &Property, 0, nullptr, &DataSize, &SampleRates);
+		CORE_AUDIO_ERR(Status, "Failed to get nominal sample rates");
+
+		// If we are given ranges of values for this device, we want to find the largest min value
+		// and the smallest max value for sample rates, then pick good sample rates based on that
+		double LargestMinSampleRateRange = 1.0f;
+		double SmallestMaxSampleRateRange = DBL_MAX;
+		bool bHasSampleRateRanges = false;
+		for (uint32 i = 0; i < NumSampleRates; ++i)
+		{
+			const AudioValueRange& SampleRate = SampleRates[i];
+
+			// If the min and max sample rate is the same, then just add this to the list of frame rates
+			if (SampleRate.mMinimum == SampleRate.mMaximum)
+			{
+				uint32 FrameRate = (uint32)SampleRate.mMinimum;
+				if (!PossibleFrameRates.Contains(FrameRate))
+				{
+					PossibleFrameRates.Add(FrameRate);
+				}
+			}
+			else
+			{
+				bHasSampleRateRanges = true;
+				if (SampleRate.mMinimum > LargestMinSampleRateRange)
+				{
+					LargestMinSampleRateRange = SampleRate.mMinimum;
+				}
+
+				if (SampleRate.mMaximum < SmallestMaxSampleRateRange)
+				{
+					SmallestMaxSampleRateRange = SampleRate.mMaximum;
+				}
+			}
+		}
+
+		if (bHasSampleRateRanges)
+		{
+			for (uint32 i = 0; i < MaxPossibleFrameRates; ++i)
+			{
+				uint32 PossibleFrameRate = PossibleFrameRates[i];
+				if (PossibleFrameRate >= LargestMinSampleRateRange && PossibleFrameRate <= SmallestMaxSampleRateRange)
+				{
+					PossibleFrameRates.Add(PossibleFrameRate);
+				}
+			}
+		}
+		// Picking a good default framerate from the supported framerates
+		if (PossibleFrameRates.Contains(48000))
+		{
+			DefaultFrameRate = 48000;
+		}
+		else if (PossibleFrameRates.Contains(44100))
+		{
+			DefaultFrameRate = 44100;
+		}
+		else
+		{
+			UA_DEVICE_PLATFORM_ERROR(TEXT("Audio device doesn't support 48k or 44.1k sample rates."));
+			return false;
+		}
+		return true;
+	}
+
+	bool FUnrealAudioCoreAudio::InitDeviceOutputId(UInt32 DeviceIndex)
+	{
+		if (DeviceIndex >= CoreAudioInfo.OutputDevices.Num())
+		{
+			return false;
+		}
+		CoreAudioInfo.OutputDeviceId = CoreAudioInfo.OutputDeviceIDs[DeviceIndex];
+		CoreAudioInfo.OutputDeviceIndex = DeviceIndex;
+		return true;
+	}
+
+	bool FUnrealAudioCoreAudio::InitDeviceFrameRate(const UInt32 RequestedFrameRate, bool& bChanged)
+	{
+		bool bSupportedFrameRate = false;
+		StreamInfo.FrameRate = RequestedFrameRate;
+		FDeviceInfo& DeviceInfo = CoreAudioInfo.OutputDevices[CoreAudioInfo.OutputDeviceIndex];
+		for (int32 i = 0; i < DeviceInfo.PossibleFrameRates.Num(); ++i)
+		{
+			if (RequestedFrameRate == DeviceInfo.PossibleFrameRates[i])
+			{
+				bSupportedFrameRate = true;
+				break;
+			}
+		}
+
+		if (!bSupportedFrameRate)
+		{
+			UA_DEVICE_PLATFORM_ERROR(TEXT("Requested frame rate is not supported by device, trying to use default."));
+			StreamInfo.FrameRate = DeviceInfo.FrameRate;
+		}
+
+		// CORE_AUDIO_ERR what the sample rate of the device is currently set to
+		Float64 NominalSampleRate;
+		UInt32 DataSize = sizeof(Float64);
+		AudioObjectPropertyAddress Property = NEW_OUTPUT_PROPERTY(kAudioDevicePropertyNominalSampleRate);
+		OSStatus Status = AudioObjectGetPropertyData(CoreAudioInfo.OutputDeviceId, &Property, 0, nullptr, &DataSize, &NominalSampleRate);
+		CORE_AUDIO_ERR(Status, "Failed to get nominal sample rate");
+
+		// If the desired sample rate is different than the nominal sample rate, then we change it
+		if (FMath::Abs(NominalSampleRate - (Float64)StreamInfo.FrameRate) > 1.0)
+		{
+			// Set a property listener to make sure we set the new sample rate correctly
+			Float64 ReportedSampleRate = 0.0;
+			AudioObjectPropertyAddress Property2 = NEW_GLOBAL_PROPERTY(kAudioDevicePropertyNominalSampleRate);
+			Status = AudioObjectAddPropertyListener(CoreAudioInfo.OutputDeviceId, &Property2, SampleRatePropertyListener, (void*)&ReportedSampleRate);
+			CORE_AUDIO_ERR(Status, "Failed to add a sample rate property listener");
+
+			NominalSampleRate = (Float64)StreamInfo.FrameRate;
+			Status = AudioObjectSetPropertyData(CoreAudioInfo.OutputDeviceId, &Property, 0, nullptr, DataSize, &NominalSampleRate);
+			CORE_AUDIO_ERR(Status, "Failed to set a sample rate on device");
+
+			// Wait until the sample rate has changed to the requested rate with a timeout
+			UInt32 Counter = 0;
+			const UInt32 Inc = 5000;
+			const UInt32 Timeout = 5000000;
+			bool bTimedout = false;
+			while (ReportedSampleRate != NominalSampleRate)
+			{
+				Counter += Inc;
+				if (Counter > Timeout)
+				{
+					bTimedout = true;
+					break;
+				}
+				usleep(Inc);
+			}
+
+			Status = AudioObjectRemovePropertyListener(CoreAudioInfo.OutputDeviceId, &Property2, SampleRatePropertyListener, (void*)&ReportedSampleRate);
+			CORE_AUDIO_ERR(Status, "Failed to remove the sample rate property listener");
+
+			if (bTimedout)
+			{
+				UA_DEVICE_PLATFORM_ERROR(TEXT("Timed out while setting sample rate of audio device."));
+				return false;
+			}
+
+			bChanged = true;
+		}
+		return true;
+	}
+
+	bool FUnrealAudioCoreAudio::InitDeviceVirtualFormat(bool bSampleRateChanged)
+	{
+		// Change the output stream virtual format if we need to
+		AudioStreamBasicDescription Format;
+		UInt32 DataSize = sizeof(AudioStreamBasicDescription);
+		AudioObjectPropertyAddress Property = NEW_OUTPUT_PROPERTY(kAudioStreamPropertyVirtualFormat);
+		OSStatus Status = AudioObjectGetPropertyData(CoreAudioInfo.OutputDeviceId, &Property, 0, nullptr, &DataSize, &Format);
+		CORE_AUDIO_ERR(Status, "Failed to get audio stream virtual format");
+
+		// we'll only change the virtual format if we need to
+		bool bChangeFormat = false;
+		if (bSampleRateChanged)
+		{
+			bChangeFormat = true;
+			Format.mSampleRate = (Float64)StreamInfo.FrameRate;
+		}
+
+		if (Format.mFormatID != kAudioFormatLinearPCM)
+		{
+			bChangeFormat = true;
+			Format.mFormatID = kAudioFormatLinearPCM;
+		}
+
+		if (bChangeFormat)
+		{
+			Status = AudioObjectSetPropertyData(CoreAudioInfo.OutputDeviceId, &Property, 0, nullptr, DataSize, &Format);
+			CORE_AUDIO_ERR(Status, "Failed to set the virtual format on device");
+		}
+		return true;
+	}
+
+	bool FUnrealAudioCoreAudio::InitDevicePhysicalFormat()
+	{
+		AudioObjectPropertyAddress Property = NEW_OUTPUT_PROPERTY(kAudioStreamPropertyPhysicalFormat);
+		AudioStreamBasicDescription Format;
+		UInt32 DataSize = sizeof(AudioStreamBasicDescription);
+		OSStatus Status = AudioObjectGetPropertyData(CoreAudioInfo.OutputDeviceId, &Property, 0, nullptr, &DataSize, &Format);
+		CORE_AUDIO_ERR(Status, "Failed to get audio stream physical format");
+
+		if (Format.mFormatID != kAudioFormatLinearPCM || Format.mBitsPerChannel < 16)
+		{
+			bool bSuccess = false;
+			Format.mFormatID = kAudioFormatLinearPCM;
+
+			// Make a copy of the old format
+			AudioStreamBasicDescription NewFormat = Format;
+
+			// Try setting the physical format to Float32
+			NewFormat.mFormatFlags = Format.mFormatFlags;
+			NewFormat.mFormatFlags |= kLinearPCMFormatFlagIsFloat;
+			NewFormat.mFormatFlags &= ~kLinearPCMFormatFlagIsSignedInteger;
+
+			NewFormat.mBitsPerChannel = 32;
+			NewFormat.mBytesPerFrame = 4 * Format.mChannelsPerFrame;
+			NewFormat.mBytesPerPacket = NewFormat.mBytesPerFrame * NewFormat.mFramesPerPacket;
+
+			Status = AudioObjectSetPropertyData(CoreAudioInfo.OutputDeviceId, &Property, 0, nullptr, DataSize, &NewFormat);
+			if (Status == noErr)
+			{
+				bSuccess = true;
+			}
+
+			// Lets try signed 32 bit integer format
+			if (!bSuccess)
+			{
+				NewFormat.mFormatFlags = Format.mFormatFlags;
+				NewFormat.mFormatFlags |= kLinearPCMFormatFlagIsSignedInteger;
+				NewFormat.mFormatFlags |= kAudioFormatFlagIsPacked;
+				NewFormat.mFormatFlags &= ~kLinearPCMFormatFlagIsFloat;
+
+				Status = AudioObjectSetPropertyData(CoreAudioInfo.OutputDeviceId, &Property, 0, nullptr, DataSize, &NewFormat);
+				if (Status == noErr)
+				{
+					bSuccess = true;
+				}
+			}
+
+			// Now try 24 bit signed integer
+			if (!bSuccess)
+			{
+				NewFormat.mBitsPerChannel = 24;
+				NewFormat.mBytesPerFrame = 3 * Format.mChannelsPerFrame;
+				NewFormat.mBytesPerPacket = NewFormat.mBytesPerFrame * NewFormat.mFramesPerPacket;
+
+				Status = AudioObjectSetPropertyData(CoreAudioInfo.OutputDeviceId, &Property, 0, nullptr, DataSize, &NewFormat);
+				if (Status == noErr)
+				{
+					bSuccess = true;
+				}
+			}
+
+			if (!bSuccess)
+			{
+				// TODO: look at supporting other physical formats if we fail here
+				UA_DEVICE_PLATFORM_ERROR(TEXT("Failed to set physical format of the audio device to something reasonable."));
+				return false;
+			}
+		}
+		return true;
+	}
+
+	bool FUnrealAudioCoreAudio::GetDeviceLatency(AudioDeviceID DeviceID, UInt32& Latency)
+	{
+		UInt32 DataSize = sizeof(UInt32);
+		AudioObjectPropertyAddress Property = NEW_OUTPUT_PROPERTY(kAudioDevicePropertyLatency);
+		if (AudioObjectHasProperty(DeviceID, &Property))
+		{
+			OSStatus Status = AudioObjectGetPropertyData(DeviceID, &Property, 0, nullptr, &DataSize, &Latency);
+			CORE_AUDIO_ERR(Status, "Failed to get device latency");
+		}
+		return true;
+	}
+
+	bool FUnrealAudioCoreAudio::InitDeviceNumDeviceStreams()
+	{
+		UInt32 DataSize;
+		AudioObjectPropertyAddress Property = NEW_OUTPUT_PROPERTY(kAudioDevicePropertyStreamConfiguration);
+		OSStatus Status = AudioObjectGetPropertyDataSize(CoreAudioInfo.OutputDeviceId, &Property, 0, nullptr, &DataSize);
+		CORE_AUDIO_ERR(Status, "Failed to get stream config size");
+
+		AudioBufferList* BufferList = (AudioBufferList*)alloca(DataSize);
+		Status = AudioObjectGetPropertyData(CoreAudioInfo.OutputDeviceId, &Property, 0, nullptr, &DataSize, BufferList);
+		CORE_AUDIO_ERR(Status, "Failed to get output stream configuration");
+
+		CoreAudioInfo.NumDeviceStreams = BufferList->mNumberBuffers;
+		return true;
+	}
+
+	bool FUnrealAudioCoreAudio::InitDeviceOverrunCallback()
+	{
+		AudioObjectPropertyAddress Property = NEW_GLOBAL_PROPERTY(kAudioDeviceProcessorOverload);
+		OSStatus Status = AudioObjectAddPropertyListener(CoreAudioInfo.OutputDeviceId, &Property, OverrunPropertyListener, (void *)this);
+		CORE_AUDIO_ERR(Status, "Failed to set device overrun property listener");
+		return true;
+	}
+
+	bool FUnrealAudioCoreAudio::InitDeviceCallback(const FCreateStreamParams& Params)
+	{
+		// Determine device callback buffer size range and set the best callback buffer size
+		AudioValueRange BufferSizeRange;
+		UInt32 DataSize = sizeof(AudioValueRange);
+		AudioObjectPropertyAddress Property = NEW_OUTPUT_PROPERTY(kAudioDevicePropertyBufferFrameSizeRange);
+		OSStatus Status = AudioObjectGetPropertyData(CoreAudioInfo.OutputDeviceId, &Property, 0, nullptr, &DataSize, &BufferSizeRange);
+		CORE_AUDIO_ERR(Status, "Failed to get callback buffer size range");
+
+		StreamInfo.BlockSize = (UInt32)FMath::Clamp(Params.CallbackBlockSize, (uint32)BufferSizeRange.mMinimum, (uint32)BufferSizeRange.mMaximum);
+
+		DataSize = sizeof(UInt32);
+		Property.mSelector = kAudioDevicePropertyBufferFrameSize;
+		Status = AudioObjectSetPropertyData(CoreAudioInfo.OutputDeviceId, &Property, 0, nullptr, DataSize, &StreamInfo.BlockSize);
+		CORE_AUDIO_ERR(Status, "Failed to set the callback buffer size");
+
+		Status = AudioDeviceCreateIOProcID(CoreAudioInfo.OutputDeviceId, CoreAudioCallback, (void*)this, &CoreAudioInfo.DeviceIOProcId);
+		CORE_AUDIO_ERR(Status, "Failed to set device callback function");
+
+		// Initialize the device byte buffer to the correct number of bytes
+		FDeviceInfo& DeviceInfo = CoreAudioInfo.OutputDevices[CoreAudioInfo.OutputDeviceIndex];
+		FStreamDeviceInfo& StreamDeviceInfo = StreamInfo.DeviceInfo;
+		StreamDeviceInfo.NumChannels = DeviceInfo.NumChannels;
+		uint64 UserBufferBytes = DeviceInfo.NumChannels * StreamInfo.BlockSize * sizeof(float);
+		StreamDeviceInfo.UserBuffer.Init(0, UserBufferBytes);
+
+		// Setup the callback info struct before starting the device stream
+		FCallbackInfo& CallbackInfo = CoreAudioInfo.CallbackInfo;
+
+		CallbackInfo.OutBuffer		= (float*)StreamDeviceInfo.UserBuffer.GetData();
+		CallbackInfo.NumFrames		= StreamInfo.BlockSize;
+		CallbackInfo.NumChannels	= DeviceInfo.NumChannels;
+		CallbackInfo.NumSamples		= CallbackInfo.NumFrames * CallbackInfo.NumChannels;
+		CallbackInfo.UserData		= Params.UserData;
+		CallbackInfo.StatusFlags	= 0;
+		CallbackInfo.OutputSpeakers	= DeviceInfo.Speakers;
+		CallbackInfo.FrameRate		= StreamInfo.FrameRate;
+		CallbackInfo.StreamTime		= 0.0;
+
+		if (pthread_cond_init(&CoreAudioInfo.Condition, NULL))
+		{
+			UA_DEVICE_PLATFORM_ERROR(TEXT("Failed to initialize thread condition"));
+			return false;
+		}
+		return true;
+	}
+
+	bool FUnrealAudioCoreAudio::GetDeviceChannelsForStereo(AudioDeviceID DeviceID, TArray<ESpeaker::Type>& OutChannels)
+	{
+		UInt32 ChannelIndices[2];
+		UInt32 PropertySize = sizeof(ChannelIndices);
+		AudioObjectPropertyAddress Property = NEW_OUTPUT_PROPERTY(kAudioDevicePropertyPreferredChannelsForStereo);
+		OSStatus Status = AudioObjectGetPropertyData(DeviceID, &Property, 0, nullptr, &PropertySize, ChannelIndices);
+		CORE_AUDIO_ERR(Status, "Failed to get preferred channels for stereo property");
+
+		if (ChannelIndices[0] == kAudioChannelLabel_Left)
+		{
+			OutChannels.Add(ESpeaker::FRONT_LEFT);
+			OutChannels.Add(ESpeaker::FRONT_RIGHT);
+		}
+		else
+		{
+			OutChannels.Add(ESpeaker::FRONT_RIGHT);
+			OutChannels.Add(ESpeaker::FRONT_LEFT);
+		}
+		return true;
+	}
+
+	bool FUnrealAudioCoreAudio::GetDeviceChannelsForLayoutDescriptions(AudioChannelLayout* ChannelLayout, TArray<ESpeaker::Type>& OutChannels)
+	{
+		bool bSuccess = true;
+		if (ChannelLayout->mNumberChannelDescriptions == 0)
+		{
+			bSuccess = false;
+		}
+		else
+		{
+			for (UInt32 Channel = 0; Channel < ChannelLayout->mNumberChannelDescriptions; ++Channel)
+			{
+				AudioChannelDescription Description = ChannelLayout->mChannelDescriptions[Channel];
+				ESpeaker::Type Speaker;
+				switch(Description.mChannelLabel)
+				{
+					case kAudioChannelLabel_Left:					Speaker = ESpeaker::FRONT_LEFT;				break;
+					case kAudioChannelLabel_Right:					Speaker = ESpeaker::FRONT_RIGHT;			break;
+					case kAudioChannelLabel_Center:					Speaker = ESpeaker::FRONT_CENTER;			break;
+					case kAudioChannelLabel_LFEScreen:				Speaker = ESpeaker::LOW_FREQUENCY;			break;
+					case kAudioChannelLabel_LeftSurround:			Speaker = ESpeaker::SIDE_LEFT;				break;
+					case kAudioChannelLabel_RightSurround:			Speaker = ESpeaker::SIDE_RIGHT;				break;
+					case kAudioChannelLabel_LeftCenter:				Speaker = ESpeaker::FRONT_LEFT_OF_CENTER;	break;
+					case kAudioChannelLabel_RightCenter:			Speaker = ESpeaker::FRONT_RIGHT_OF_CENTER;	break;
+					case kAudioChannelLabel_CenterSurround:			Speaker = ESpeaker::BACK_CENTER;			break;
+					case kAudioChannelLabel_LeftSurroundDirect:		Speaker = ESpeaker::SIDE_LEFT;				break;
+					case kAudioChannelLabel_RightSurroundDirect:	Speaker = ESpeaker::SIDE_RIGHT;				break;
+					case kAudioChannelLabel_TopCenterSurround:		Speaker = ESpeaker::TOP_CENTER;				break;
+					case kAudioChannelLabel_VerticalHeightLeft:		Speaker = ESpeaker::TOP_FRONT_LEFT;			break;
+					case kAudioChannelLabel_VerticalHeightCenter:	Speaker = ESpeaker::TOP_FRONT_CENTER;		break;
+					case kAudioChannelLabel_VerticalHeightRight:	Speaker = ESpeaker::TOP_FRONT_RIGHT;		break;
+					case kAudioChannelLabel_TopBackLeft:			Speaker = ESpeaker::TOP_BACK_LEFT;			break;
+					case kAudioChannelLabel_TopBackCenter:			Speaker = ESpeaker::TOP_BACK_CENTER;		break;
+					case kAudioChannelLabel_TopBackRight:			Speaker = ESpeaker::TOP_BACK_RIGHT;			break;
+					case kAudioChannelLabel_RearSurroundLeft:		Speaker = ESpeaker::BACK_LEFT;				break;
+					case kAudioChannelLabel_RearSurroundRight:		Speaker = ESpeaker::BACK_RIGHT;				break;
+					case kAudioChannelLabel_Unused:					Speaker = ESpeaker::UNUSED;					break;
+					case kAudioChannelLabel_Unknown:				Speaker = ESpeaker::UNUSED;					break;
+
+					default:
+						UA_DEVICE_PLATFORM_ERROR(TEXT("Unknown or unsupported channel label"));
+						OutChannels.Empty();
+						return false;
+				}
+				OutChannels.Add(Speaker);
+			}
+		}
+		return bSuccess;
+	}
+
+	bool FUnrealAudioCoreAudio::GetDeviceChannelsForBitMap(UInt32 BitMap, TArray<ESpeaker::Type>& OutChannels)
+	{
+		bool bSuccess = true;
+		// Bit maps for standard speaker layouts
+		static const UInt32 BitMapMono				= kAudioChannelBit_Center;
+		static const UInt32 BitMapStereo			= kAudioChannelBit_Left | kAudioChannelBit_Right;
+		static const UInt32 BitMapStereoPoint1		= kAudioChannelBit_Left | kAudioChannelBit_Right | kAudioChannelBit_LFEScreen;
+		static const UInt32 BitMapSurround			= kAudioChannelBit_Left | kAudioChannelBit_Right | kAudioChannelBit_Center | kAudioChannelBit_CenterSurround;
+		static const UInt32 BitMapQuad				= kAudioChannelBit_Left | kAudioChannelBit_Right | kAudioChannelBit_LeftSurround | kAudioChannelBit_RightSurround;
+		static const UInt32 BitMap4Point1			= kAudioChannelBit_Left | kAudioChannelBit_Right | kAudioChannelBit_LeftSurround | kAudioChannelBit_RightSurround | kAudioChannelBit_LFEScreen;
+		static const UInt32 BitMap5Point1			= kAudioChannelBit_Left | kAudioChannelBit_Right | kAudioChannelBit_Center | kAudioChannelBit_LeftSurround | kAudioChannelBit_RightSurround | kAudioChannelBit_LFEScreen;
+		static const UInt32 BitMap7Point1			= kAudioChannelBit_Left | kAudioChannelBit_Right | kAudioChannelBit_Center | kAudioChannelBit_LeftSurround | kAudioChannelBit_RightSurround | kAudioChannelBit_LFEScreen | kAudioChannelBit_LeftCenter | kAudioChannelBit_RightCenter;
+		static const UInt32 BitMap5Point1Surround	= kAudioChannelBit_Left | kAudioChannelBit_Right | kAudioChannelBit_Center | kAudioChannelBit_LeftSurroundDirect | kAudioChannelBit_RightSurroundDirect | kAudioChannelBit_LFEScreen;
+		static const UInt32 BitMap7Point1Surround	= kAudioChannelBit_Left | kAudioChannelBit_Right | kAudioChannelBit_Center | kAudioChannelBit_LeftSurroundDirect | kAudioChannelBit_RightSurroundDirect | kAudioChannelBit_LFEScreen | kAudioChannelBit_LeftSurround | kAudioChannelBit_RightSurround;
+
+		switch(BitMap)
+		{
+			case BitMapMono:
+				OutChannels.Add(ESpeaker::FRONT_CENTER);
+				break;
+
+			case BitMapStereo:
+				OutChannels.Add(ESpeaker::FRONT_LEFT);
+				OutChannels.Add(ESpeaker::FRONT_RIGHT);
+				break;
+
+			case BitMapStereoPoint1:
+				OutChannels.Add(ESpeaker::FRONT_LEFT);
+				OutChannels.Add(ESpeaker::FRONT_RIGHT);
+				OutChannels.Add(ESpeaker::LOW_FREQUENCY);
+				break;
+
+			case BitMapSurround:
+				OutChannels.Add(ESpeaker::FRONT_LEFT);
+				OutChannels.Add(ESpeaker::FRONT_RIGHT);
+				OutChannels.Add(ESpeaker::FRONT_CENTER);
+				OutChannels.Add(ESpeaker::BACK_CENTER);
+				break;
+
+			case BitMapQuad:
+				OutChannels.Add(ESpeaker::FRONT_LEFT);
+				OutChannels.Add(ESpeaker::FRONT_RIGHT);
+				OutChannels.Add(ESpeaker::BACK_LEFT);
+				OutChannels.Add(ESpeaker::BACK_RIGHT);
+				break;
+
+			case BitMap4Point1:
+				OutChannels.Add(ESpeaker::FRONT_LEFT);
+				OutChannels.Add(ESpeaker::FRONT_RIGHT);
+				OutChannels.Add(ESpeaker::BACK_LEFT);
+				OutChannels.Add(ESpeaker::BACK_RIGHT);
+				OutChannels.Add(ESpeaker::LOW_FREQUENCY);
+				break;
+
+			case BitMap5Point1:
+				OutChannels.Add(ESpeaker::FRONT_LEFT);
+				OutChannels.Add(ESpeaker::FRONT_RIGHT);
+				OutChannels.Add(ESpeaker::FRONT_CENTER);
+				OutChannels.Add(ESpeaker::LOW_FREQUENCY);
+				OutChannels.Add(ESpeaker::BACK_LEFT);
+				OutChannels.Add(ESpeaker::BACK_RIGHT);
+				break;
+
+			case BitMap7Point1:
+				OutChannels.Add(ESpeaker::FRONT_LEFT);
+				OutChannels.Add(ESpeaker::FRONT_RIGHT);
+				OutChannels.Add(ESpeaker::FRONT_CENTER);
+				OutChannels.Add(ESpeaker::LOW_FREQUENCY);
+				OutChannels.Add(ESpeaker::BACK_LEFT);
+				OutChannels.Add(ESpeaker::BACK_RIGHT);
+				OutChannels.Add(ESpeaker::FRONT_LEFT_OF_CENTER);
+				OutChannels.Add(ESpeaker::FRONT_RIGHT_OF_CENTER);
+				break;
+
+			case BitMap5Point1Surround:
+				OutChannels.Add(ESpeaker::FRONT_LEFT);
+				OutChannels.Add(ESpeaker::FRONT_RIGHT);
+				OutChannels.Add(ESpeaker::FRONT_CENTER);
+				OutChannels.Add(ESpeaker::LOW_FREQUENCY);
+				OutChannels.Add(ESpeaker::SIDE_LEFT);
+				OutChannels.Add(ESpeaker::SIDE_RIGHT);
+				break;
+
+			case BitMap7Point1Surround:
+				OutChannels.Add(ESpeaker::FRONT_LEFT);
+				OutChannels.Add(ESpeaker::FRONT_RIGHT);
+				OutChannels.Add(ESpeaker::FRONT_CENTER);
+				OutChannels.Add(ESpeaker::LOW_FREQUENCY);
+				OutChannels.Add(ESpeaker::BACK_LEFT);
+				OutChannels.Add(ESpeaker::BACK_RIGHT);
+				OutChannels.Add(ESpeaker::SIDE_LEFT);
+				OutChannels.Add(ESpeaker::SIDE_RIGHT);
+				break;
+
+			default:
+				UA_DEVICE_PLATFORM_ERROR(TEXT("Unknown or unsupported channel bitmap"));
+				bSuccess = false;
+				break;
+		}
+		return bSuccess;
+	}
+
+	bool FUnrealAudioCoreAudio::GetDeviceChannelsForLayoutTag(AudioChannelLayoutTag LayoutTag, TArray<ESpeaker::Type>& OutChannels)
+	{
+		bool bSuccess = true;
+		switch (LayoutTag)
+		{
+			case kAudioChannelLayoutTag_Mono:
+				OutChannels.Add(ESpeaker::FRONT_CENTER);
+				break;
+
+			case kAudioChannelLayoutTag_Stereo:
+			case kAudioChannelLayoutTag_StereoHeadphones:
+			case kAudioChannelLayoutTag_MatrixStereo:
+			case kAudioChannelLayoutTag_MidSide:
+			case kAudioChannelLayoutTag_XY:
+			case kAudioChannelLayoutTag_Binaural:
+				OutChannels.Add(ESpeaker::FRONT_LEFT);
+				OutChannels.Add(ESpeaker::FRONT_RIGHT);
+				break;
+
+			case kAudioChannelLayoutTag_Quadraphonic:
+				OutChannels.Add(ESpeaker::FRONT_LEFT);
+				OutChannels.Add(ESpeaker::FRONT_RIGHT);
+				OutChannels.Add(ESpeaker::BACK_LEFT);
+				OutChannels.Add(ESpeaker::BACK_RIGHT);
+				break;
+
+			case kAudioChannelLayoutTag_Pentagonal:
+				OutChannels.Add(ESpeaker::FRONT_LEFT);
+				OutChannels.Add(ESpeaker::FRONT_RIGHT);
+				OutChannels.Add(ESpeaker::BACK_LEFT);
+				OutChannels.Add(ESpeaker::BACK_RIGHT);
+				OutChannels.Add(ESpeaker::FRONT_CENTER);
+				break;
+
+			case kAudioChannelLayoutTag_Hexagonal:
+				OutChannels.Add(ESpeaker::FRONT_LEFT);
+				OutChannels.Add(ESpeaker::FRONT_RIGHT);
+				OutChannels.Add(ESpeaker::BACK_LEFT);
+				OutChannels.Add(ESpeaker::BACK_RIGHT);
+				OutChannels.Add(ESpeaker::FRONT_CENTER);
+				OutChannels.Add(ESpeaker::LOW_FREQUENCY);
+				break;
+
+			case kAudioChannelLayoutTag_Octagonal:
+				OutChannels.Add(ESpeaker::FRONT_LEFT);
+				OutChannels.Add(ESpeaker::FRONT_RIGHT);
+				OutChannels.Add(ESpeaker::BACK_LEFT);
+				OutChannels.Add(ESpeaker::BACK_RIGHT);
+				OutChannels.Add(ESpeaker::FRONT_CENTER);
+				OutChannels.Add(ESpeaker::LOW_FREQUENCY);
+				OutChannels.Add(ESpeaker::SIDE_LEFT);
+				OutChannels.Add(ESpeaker::SIDE_RIGHT);
+				break;
+
+			default:
+				bSuccess = false;
+				break;
+		}
+		return bSuccess;
+	}
+
+	bool FUnrealAudioCoreAudio::GetDeviceChannelsForChannelCount(UInt32 NumChannels, TArray<ESpeaker::Type>& OutChannels)
+	{
+		bool bSuccess = true;
+		switch (NumChannels)
+		{
+			case 1:
+				OutChannels.Add(ESpeaker::FRONT_CENTER);
+				break;
+
+			case 2:
+				OutChannels.Add(ESpeaker::FRONT_LEFT);
+				OutChannels.Add(ESpeaker::FRONT_RIGHT);
+				break;
+
+			case 3:
+				OutChannels.Add(ESpeaker::FRONT_LEFT);
+				OutChannels.Add(ESpeaker::FRONT_RIGHT);
+				OutChannels.Add(ESpeaker::LOW_FREQUENCY);
+				break;
+
+			case 4:
+				OutChannels.Add(ESpeaker::FRONT_LEFT);
+				OutChannels.Add(ESpeaker::FRONT_RIGHT);
+				OutChannels.Add(ESpeaker::BACK_LEFT);
+				OutChannels.Add(ESpeaker::BACK_RIGHT);
+				break;
+
+			case 5:
+				OutChannels.Add(ESpeaker::FRONT_LEFT);
+				OutChannels.Add(ESpeaker::FRONT_RIGHT);
+				OutChannels.Add(ESpeaker::BACK_LEFT);
+				OutChannels.Add(ESpeaker::BACK_RIGHT);
+				OutChannels.Add(ESpeaker::FRONT_CENTER);
+				break;
+
+			case 6:
+				OutChannels.Add(ESpeaker::FRONT_LEFT);
+				OutChannels.Add(ESpeaker::FRONT_RIGHT);
+				OutChannels.Add(ESpeaker::BACK_LEFT);
+				OutChannels.Add(ESpeaker::BACK_RIGHT);
+				OutChannels.Add(ESpeaker::FRONT_CENTER);
+				OutChannels.Add(ESpeaker::LOW_FREQUENCY);
+				break;
+
+			case 8:
+				OutChannels.Add(ESpeaker::FRONT_LEFT);
+				OutChannels.Add(ESpeaker::FRONT_RIGHT);
+				OutChannels.Add(ESpeaker::BACK_LEFT);
+				OutChannels.Add(ESpeaker::BACK_RIGHT);
+				OutChannels.Add(ESpeaker::FRONT_CENTER);
+				OutChannels.Add(ESpeaker::LOW_FREQUENCY);
+				OutChannels.Add(ESpeaker::SIDE_LEFT);
+				OutChannels.Add(ESpeaker::SIDE_RIGHT);
+				break;
+
+			default:
+				bSuccess = false;
+				UA_DEVICE_PLATFORM_ERROR(TEXT("Failed to get a speaker array from number of channels"));
+				break;
+		}
+		return bSuccess;
+	}
+
+
+}
+
+IMPLEMENT_MODULE(UAudio::FUnrealAudioCoreAudio, UnrealAudioCoreAudio);
+
+#endif
\ No newline at end of file
diff --git a/Engine/Source/Runtime/Mac/UnrealAudioCoreAudio/UnrealAudioCoreAudio.Build.cs b/Engine/Source/Runtime/Mac/UnrealAudioCoreAudio/UnrealAudioCoreAudio.Build.cs
new file mode 100644
--- /dev/null
+++ b/Engine/Source/Runtime/Mac/UnrealAudioCoreAudio/UnrealAudioCoreAudio.Build.cs
@@ -0,0 +1,33 @@
+// Copyright 1998-2015 Epic Games, Inc. All Rights Reserved.
+
+using UnrealBuildTool;
+
+public class UnrealAudioCoreAudio : ModuleRules
+{
+	public UnrealAudioCoreAudio(TargetInfo Target)
+	{
+		PrivateIncludePaths.AddRange(
+			new string[] {
+				"Runtime/UnrealAudio/Private"
+			}
+		);
+
+		PrivateIncludePathModuleNames.Add("TargetPlatform");
+
+		PrivateDependencyModuleNames.AddRange(
+			new string[] {
+				"Core",
+				"CoreUObject",
+				"Engine",
+			}
+		);
+
+		PublicFrameworks.AddRange(
+			new string[] { 
+				"CoreAudio", 
+				"AudioToolbox", 
+				"AudioUnit" 
+			}
+		);
+	}
+}
diff --git a/Engine/Source/Runtime/UnrealAudio/Private/Tests/UnrealAudioDeviceTests.cpp b/Engine/Source/Runtime/UnrealAudio/Private/Tests/UnrealAudioDeviceTests.cpp
--- a/Engine/Source/Runtime/UnrealAudio/Private/Tests/UnrealAudioDeviceTests.cpp
+++ b/Engine/Source/Runtime/UnrealAudio/Private/Tests/UnrealAudioDeviceTests.cpp
@@ -37,7 +37,7 @@ namespace UAudio
 		}
 
 		
-		UE_LOG(LogUnrealAudio, Display, TEXT("Querying which audio device API we're using..."), );
+		UE_LOG(LogUnrealAudio, Display, TEXT("Querying which audio device API we're using..."));
 		EDeviceApi::Type DeviceApi;
 		if (!UnrealAudioDevice->GetDevicePlatformApi(DeviceApi))
 		{
@@ -54,16 +54,16 @@ namespace UAudio
 			return true;
 		}
 
-		UE_LOG(LogUnrealAudio, Display, TEXT("Querying the number of output devices for current system..."), );
+		UE_LOG(LogUnrealAudio, Display, TEXT("Querying the number of output devices for current system..."));
 		uint32 NumOutputDevices = 0;
 		if (!UnrealAudioDevice->GetNumOutputDevices(NumOutputDevices))
 		{
 			UE_LOG(LogUnrealAudio, Display, TEXT("Failed."));
 			return false;
 		}
-		UE_LOG(LogUnrealAudio, Display, TEXT("Success: %d Number of Output Devices"), NumOutputDevices);
+		UE_LOG(LogUnrealAudio, Display, TEXT("Success: %d Output Devices Found"), NumOutputDevices);
 
-		UE_LOG(LogUnrealAudio, Display, TEXT("Retrieving output device info for each output device..."), );
+		UE_LOG(LogUnrealAudio, Display, TEXT("Retrieving output device info for each output device..."));
 		for (uint32 DeviceIndex = 0; DeviceIndex < NumOutputDevices; ++DeviceIndex)
 		{
 			FDeviceInfo DeviceInfo;
@@ -188,7 +188,7 @@ namespace UAudio
 		// Block until the synthesizer is done
 		while (!Generator->IsDone())
 		{
-			Sleep(1);
+			FPlatformProcess::Sleep(1);
 		}
 
 		// Stop the output stream (which blocks until its actually freed)
@@ -213,7 +213,7 @@ namespace UAudio
 	bool TestDeviceOutputSimple(double LifeTime)
 	{
 		Test::FSimpleOutput SimpleOutput(LifeTime);
-		return DoOutputTest("output randomized FM synthesis", LifeTime, &SimpleOutput);
+		return DoOutputTest("output simple test", LifeTime, &SimpleOutput);
 	}
 
 	bool TestDeviceOutputRandomizedFm(double LifeTime)
diff --git a/Engine/Source/Runtime/UnrealAudio/Private/Tests/UnrealAudioTestGenerators.cpp b/Engine/Source/Runtime/UnrealAudio/Private/Tests/UnrealAudioTestGenerators.cpp
--- a/Engine/Source/Runtime/UnrealAudio/Private/Tests/UnrealAudioTestGenerators.cpp
+++ b/Engine/Source/Runtime/UnrealAudio/Private/Tests/UnrealAudioTestGenerators.cpp
@@ -44,7 +44,16 @@ namespace Test
 
 	static float WrapTwoPi(float Value)
 	{
-		return FMath::UnwindRadians(Value) + PI;
+		while (Value > TWO_PI)
+		{
+			Value -= TWO_PI;
+		}
+
+		while (Value < 0)
+		{
+			Value += TWO_PI;
+		}
+		return Value;
 	}
 
 	static float Lerp(float StartX, float EndX, float StartY, float EndY, float Alpha)
@@ -98,12 +107,7 @@ namespace Test
 
 	bool FTimer::IsDone() const
 	{
-		bool bIsDone = TotalTime >= 0.0f && CurrentTime >= TotalTime;
-		if (bIsDone)
-		{
-			printf("test");
-		}
-		return bIsDone;
+		return TotalTime >= 0.0f && CurrentTime >= TotalTime;;
 	}
 
 	void FTimer::Start(double InTotalTime)
@@ -123,63 +127,6 @@ namespace Test
 		return TotalTime;
 	}
 
-	/************************************************************************/
-	/* FFader																*/
-	/************************************************************************/
-
-	FFader::FFader(float FadeInEase, float FadeOutEase)
-		: FadeInEase(FadeInEase)
-		, FadeOutEase(FadeOutEase)
-		, FadeValue(0.0f)
-		, FadeState(EFadeState::OFF)
-	{
-	}
-
-	FFader::~FFader()
-	{
-	}
-
-	float FFader::Update()
-	{
-		if (FadeState == EFadeState::FADE_IN)
-		{
-			FadeValue += (1.0f - FadeValue) * FadeInEase;
-			if (1.0f - FadeValue < 0.0001f)
-			{
-				FadeState = EFadeState::ON;
-				FadeValue = 1.0f;
-			}
-		}
-		else if (FadeState == EFadeState::FADE_OUT)
-		{
-			FadeValue += (0.0f - FadeValue) * FadeOutEase;
-			if (FadeValue < 0.0001f)
-			{
-				FadeValue = 0.0f;
-				FadeState = EFadeState::OFF;
-			}
-		}
-		return FadeValue;
-
-	}
-
-	void FFader::Fade(bool bFadeIn)
-	{
-		if (bFadeIn)
-		{
-			FadeState = EFadeState::FADE_IN;
-		}
-		else
-		{
-			FadeState = EFadeState::FADE_OUT;
-		}
-	}
-
-	void FFader::Toggle()
-	{
-		Fade(FadeState == EFadeState::FADE_OUT || FadeState == EFadeState::OFF);
-	}
-
 	/************************************************************************/
 	/* FBiquad																*/
 	/************************************************************************/
@@ -215,8 +162,8 @@ namespace Test
 	/************************************************************************/
 
 	FLowPass::FLowPass()
-		: Frequency(0.0f)
-		, Quality(1.0f)
+		: Quality(1.0f)
+		, Frequency(0.0f)
 	{
 	}
 
@@ -257,10 +204,54 @@ namespace Test
 	}
 
 	/************************************************************************/
-	/* FRamp																*/
+	/* FLinEnv														*/
 	/************************************************************************/
 
-	FRamp::FRamp()
+	FLinEnv::FLinEnv()
+		: CurrValue(0.0f)
+		, EndValue(0.0f)
+		, bIsDone(false)
+	{
+	}
+
+	FLinEnv::~FLinEnv()
+	{
+	}
+
+	void FLinEnv::Start(float End, float Time)
+	{
+		if (Time == 0.0f)
+		{
+			bIsDone = true;
+			CurrValue = End;
+		}
+		else
+		{
+			bIsDone = false;
+			EndValue = End;
+			DeltaValue = (End - CurrValue) / (Time * CallbackData.FrameRate);
+		}
+	}
+
+	float FLinEnv::Update()
+	{
+		if (FMath::Abs(CurrValue - EndValue) < 0.001f)
+		{
+			bIsDone = true;
+			CurrValue = EndValue;
+		}
+		else
+		{
+			CurrValue += DeltaValue;
+		}
+		return CurrValue;
+	}
+
+	/************************************************************************/
+	/* FSaw																	*/
+	/************************************************************************/
+
+	FSaw::FSaw()
 		: CurrValue(0.0f)
 		, DeltaValue(0.0f)
 		, Frequency(0.0f)
@@ -269,35 +260,36 @@ namespace Test
 	{
 	}
 
-	FRamp::~FRamp()
+	FSaw::~FSaw()
 	{
 	}
 
-	void FRamp::SetFrequency(float InFrequency)
+	void FSaw::SetFrequency(float InFrequency)
 	{
 		check(CallbackData.FrameRate > 0.0f);
 		Frequency = InFrequency;
 		DeltaValue = Frequency / CallbackData.FrameRate;
 	}
 
-	void FRamp::SetScaleAdd(float InScale, float InAdd)
+	void FSaw::SetScaleAdd(float InScale, float InAdd)
 	{
 		Scale = InScale;
 		Add = InAdd;
 	}
 
-	float FRamp::Update()
+	float FSaw::Update()
 	{
+		float Result = CurrValue;
 		CurrValue += DeltaValue;
-		if (CurrValue >= 1.0f)
+		while (CurrValue >= 1.0f)
 		{
 			CurrValue -= 1.0f;
 		}
-		else if (CurrValue <= 0.0f)
+		while (CurrValue <= 0.0f)
 		{
 			CurrValue += 1.0f;
 		}
-		return Scale * CurrValue + Add;
+		return Scale * Result + Add;
 	}
 
 	/************************************************************************/
@@ -410,7 +402,7 @@ namespace Test
 
 	void FPan::SetPan(float InPan)
 	{
-		Pan = InPan;
+		Pan = FMath::Clamp(InPan, 0.0f, 1.0f);
 	}
 
 	void FPan::Spatialize(float Value, float* Frame)
@@ -670,11 +662,8 @@ namespace Test
 
 	FSimpleOutput::FSimpleOutput(double LifeTime)
 		: IGenerator(LifeTime)
-		, BaseFreqHz(220.0f)
-		, PhaseDelta(0.0f)
 		, Amplitude(0.5f)
-		, NumNonLfeSpeakers(0)
-		, LFEIndex(-1)
+		, ChannelTimer(1.0f)
 	{
 	}
 
@@ -682,11 +671,6 @@ namespace Test
 	{
 	}
 
-	bool FSimpleOutput::IsInit()
-	{
-		return NumNonLfeSpeakers == 0;
-	}
-
 	bool FSimpleOutput::GetNextBuffer(FCallbackInfo& CallbackInfo)
 	{
 		if (!UpdateTimers())
@@ -694,58 +678,19 @@ namespace Test
 			return true;
 		}
 
-		// phase delta is radians per frame, or amount to increment the current phase each frame
-		if (IsInit())
-		{
-			check(CallbackData.FrameRate > 0);
-			PhaseDelta = (TWO_PI * BaseFreqHz) / CallbackData.FrameRate;
-			NumNonLfeSpeakers = 0;
-			for (int32 Chan = 0; Chan < CallbackInfo.NumChannels; ++Chan)
-			{
-				if (CallbackInfo.OutputSpeakers[Chan] == ESpeaker::LOW_FREQUENCY)
-				{
-					LFEIndex = Chan;
-				}
-				else
-				{
-					++NumNonLfeSpeakers;
-				}
-			}
-
-			SpeakerInfo.Init(FSpeakerInfo(), NumNonLfeSpeakers);
-			for (int32 Index = 0; Index < NumNonLfeSpeakers; ++Index)
-			{
-				SpeakerInfo[Index].MaxHarmonic = FMath::RandRange(3, 12);
-				SpeakerInfo[Index].Time = (float) FMath::RandRange(1, 8) / 16.0f;
-			}
-		}
-
-		for (int32 Index = 0; Index < SpeakerInfo.Num(); ++Index)
+		if (Sinusoids.Num() == 0)
 		{
-			FSpeakerInfo& Info = SpeakerInfo[Index++];
-			if (Info.Timer.Update())
+			ChannelTimer.Start(1.0f);
+			CurrentChannelOut = 0;
+			for (int32 i = 0; i < CallbackInfo.NumChannels; ++i)
 			{
-				Info.bOn = !Info.bOn;
-				if (Info.bOn)
+				FChannelSine ChannelSine;
+				ChannelSine.SineOsc.SetFrequency(220.0f * (1 + i));
+				if (i == CurrentChannelOut)
 				{
-					if (Info.bDirection)
-					{
-						Info.Harmonic++;
-					}
-					else
-					{
-						Info.Harmonic--;
-					}
-					if (Info.Harmonic >= Info.MaxHarmonic)
-					{
-						Info.bDirection = false;
-					}
-					else if (Info.Harmonic == 1)
-					{
-						Info.bDirection = true;
-					}
+					ChannelSine.Fader.Start(1.0f, 0.1f);
 				}
-				Info.Timer.Start(Info.Time);
+				Sinusoids.Add(ChannelSine);
 			}
 		}
 
@@ -755,24 +700,24 @@ namespace Test
 			int32 SpeakerIndex = 0;
 			for (int32 Channel = 0; Channel < CallbackInfo.NumChannels; ++Channel)
 			{
-				// Ignore the LFE speaker
-				if (Channel == LFEIndex)
+				FSineOsc& SineOsc = Sinusoids[Channel].SineOsc;
+				FLinEnv& Fader = Sinusoids[Channel].Fader;
+				float Value = SineOsc.Update();
+				Value *= Fader.Update();
+				Value *= Amplitude;
+				CallbackInfo.OutBuffer[FrameIndex + Channel] = Value;
+
+				if (Fader.IsDone() && Channel == CurrentChannelOut)
 				{
-					continue;
+					Fader.Start(0.0f, 0.5f);
 				}
+			}
 
-				FSpeakerInfo& Info = SpeakerInfo[SpeakerIndex++];
-				float Value = Amplitude * FMath::Sin(Info.Phase);
-
-				Info.Fader.Fade(Info.bOn);
-
-				CallbackInfo.OutBuffer[FrameIndex + Channel] = Info.Fader.Update()*Value;
-
-				// Increment the phase, an octave up for each channel
-				Info.Phase += (PhaseDelta * (float)Info.Harmonic);
-
-				// Wrap the phase
-				Info.Phase = WrapTwoPi(Info.Phase);
+			if (ChannelTimer.Update())
+			{
+				CurrentChannelOut = (CurrentChannelOut + 1) % CallbackInfo.NumChannels;
+				Sinusoids[CurrentChannelOut].Fader.Start(1.0f, 0.1f);
+				ChannelTimer.Start(1.0f);
 			}
 		}
 		return true;
diff --git a/Engine/Source/Runtime/UnrealAudio/Private/Tests/UnrealAudioTestGenerators.h b/Engine/Source/Runtime/UnrealAudio/Private/Tests/UnrealAudioTestGenerators.h
--- a/Engine/Source/Runtime/UnrealAudio/Private/Tests/UnrealAudioTestGenerators.h
+++ b/Engine/Source/Runtime/UnrealAudio/Private/Tests/UnrealAudioTestGenerators.h
@@ -43,36 +43,46 @@ namespace Test
 	};
 
 	/**
-	* FFader
-	* A simple sample-accurate fader
+	* FLinEnv
+	* A simple linear envelope
 	*/
-	class FFader
+	class FLinEnv
 	{
+	public:
+		FLinEnv();
+		~FLinEnv();
+
+		void Start(float End, float TimeSec);
+		float Update();
+		bool IsDone() const { return bIsDone; }
+
 	private:
-		struct EFadeState
-		{
-			enum Type
-			{
-				FADE_IN,
-				FADE_OUT,
-				OFF,
-				ON
-			};
-		};
+		float CurrValue;
+		float DeltaValue;
+		float EndValue;
+		bool bIsDone;
+	};
 
+	/**
+	* FSaw
+	* A simple saw-wave class
+	*/
+	class FSaw
+	{
 	public:
-		FFader(float FadeInEase, float FadeOutEase);
-		~FFader();
+		FSaw();
+		~FSaw();
 
+		void SetFrequency(float InFrequency);
+		void SetScaleAdd(float InScale, float InAdd);
 		float Update();
-		void Fade(bool bFadeIn);
-		void Toggle();
 
 	private:
-		float FadeInEase;
-		float FadeOutEase;
-		float FadeValue;
-		EFadeState::Type FadeState;
+		float CurrValue;
+		float DeltaValue;
+		float Frequency;
+		float Scale;
+		float Add;
 	};
 
 	/**
@@ -119,27 +129,6 @@ namespace Test
 		float Frequency;
 	};
 
-	/**
-	* FRamp
-	* A linear time-ramp class.
-	*/
-	class FRamp
-	{
-	public:
-		FRamp();
-		~FRamp();
-
-		void SetFrequency(float InFrequency);
-		void SetScaleAdd(float InScale, float InAdd);
-		float Update();
-
-	private:
-		float CurrValue;
-		float DeltaValue;
-		float Frequency;
-		float Scale;
-		float Add;
-	};
 
 	/**
 	* FSineOsc
@@ -272,38 +261,16 @@ namespace Test
 		bool GetNextBuffer(FCallbackInfo& CallbackInfo) override;
 
 	private:
-		bool IsInit();
-
-		struct FSpeakerInfo
+		struct FChannelSine
 		{
-			float Phase;
-			FFader Fader;
-			int32 Harmonic;
-			int32 MaxHarmonic;
-			FTimer Timer;
-			float Time;
-			bool bDirection;
-			bool bOn;
-
-			FSpeakerInfo()
-				: Phase(0.0f)
-				, Fader(0.01f, 0.0001f)
-				, Harmonic(1)
-				, MaxHarmonic(10)
-				, Timer(1.0f)
-				, Time(1.0f)
-				, bDirection(true)
-				, bOn(false)
-			{
-			}
+			FSineOsc SineOsc;
+			FLinEnv Fader;
 		};
 
-		float BaseFreqHz;
-		float PhaseDelta;
+		TArray<FChannelSine> Sinusoids;
 		float Amplitude;
-		TArray<FSpeakerInfo> SpeakerInfo;
-		int32 NumNonLfeSpeakers;
-		int32 LFEIndex;
+		uint32 CurrentChannelOut;
+		FTimer ChannelTimer;
 	};
 
 
@@ -362,7 +329,7 @@ namespace Test
 			TArray<float> SpeakerMap;
 			FLowPass LowPass;
 			FPan Panner;
-			FRamp PanRamp;
+			FSaw PanRamp;
 			FSineOsc FilterLFO;
 			bool bInit;
 		};
@@ -390,7 +357,7 @@ namespace Test
 		float Amp;
 		FLowPass LowPass;
 		FPan Panner;
-		FRamp PanRamp;
+		FSaw PanRamp;
 		FSineOsc FilterLFO;
 		FDelay Delay;
 		bool bInit;
diff --git a/Engine/Source/Runtime/UnrealAudio/Private/UnrealAudio.cpp b/Engine/Source/Runtime/UnrealAudio/Private/UnrealAudio.cpp
--- a/Engine/Source/Runtime/UnrealAudio/Private/UnrealAudio.cpp
+++ b/Engine/Source/Runtime/UnrealAudio/Private/UnrealAudio.cpp
@@ -7,8 +7,11 @@
 
 #if ENABLE_UNREAL_AUDIO
 
-#define AUDIO_DEVICE_DEVICE_MODULE_NAME "UnrealAudioWasapi"
-//#define AUDIO_DEVICE_DEVICE_MODULE_NAME TEXT("UnrealAudioXAudio2")
+#if PLATFORM_WINDOWS
+#define AUDIO_DEFAULT_DEVICE_MODULE_NAME "UnrealAudioXAudio2"
+#elif PLATFORM_MAC
+#define AUDIO_DEFAULT_DEVICE_MODULE_NAME "UnrealAudioCoreAudio"
+#endif
 
 DEFINE_LOG_CATEGORY(LogUnrealAudio);
 
@@ -71,7 +74,7 @@ namespace UAudio
 		}
 
 		EDeviceApi::Type DeviceApi;
-		if (UnrealAudioDevice->GetDevicePlatformApi(DeviceApi) && DeviceApi == EDeviceApi::DUMMY)
+		if (UnrealAudioDevice && UnrealAudioDevice->GetDevicePlatformApi(DeviceApi) && DeviceApi == EDeviceApi::DUMMY)
 		{
 			delete UnrealAudioDevice;
 			UnrealAudioDevice = nullptr;
@@ -82,7 +85,7 @@ namespace UAudio
 	FName FUnrealAudioModule::GetDefaultDeviceModuleName() const
 	{
 		// TODO: Pull the device module name to use from an engine ini or an audio ini file
-		return FName(AUDIO_DEVICE_DEVICE_MODULE_NAME);
+		return FName(AUDIO_DEFAULT_DEVICE_MODULE_NAME);
 	}
 
 
diff --git a/Engine/Source/Runtime/UnrealAudio/Private/UnrealAudioDevice.cpp b/Engine/Source/Runtime/UnrealAudio/Private/UnrealAudioDevice.cpp
--- a/Engine/Source/Runtime/UnrealAudio/Private/UnrealAudioDevice.cpp
+++ b/Engine/Source/Runtime/UnrealAudio/Private/UnrealAudioDevice.cpp
@@ -1,6 +1,7 @@
 // Copyright 1998-2015 Epic Games, Inc. All Rights Reserved.
 
 #include "UnrealAudioPrivate.h"
+#include "UnrealAudioModule.h"
 #include "UnrealAudioDeviceModule.h"
 
 #if ENABLE_UNREAL_AUDIO
@@ -9,71 +10,6 @@ DEFINE_LOG_CATEGORY(LogUnrealAudioDevice);
 
 namespace UAudio
 {
-	/**
-	Static array specifying which frame rates to check for support on audio devices.
-	http://en.wikipedia.org/wiki/Sampling_%28signal_processing%29
-	*/
-	const uint32 IUnrealAudioDeviceModule::PossibleFrameRates[] =
-	{
-		8000,		/** Used for telephony, walkie talkies, really shitty but ok for human speech. */
-		11025,		/** Quarter sample rate of CD's, low-quality PCM */
-		16000,		/** Rate used for VOIP (which is why VOIP sounds slightly better than normal phones), wide-band extension over normal telephony */
-		22050,		/** Half CD sample rate, low quality PCM */
-		32000,		/** MiniDV, digital FM radio, decent wireless microphones */
-		44100,		/** CD's, MPEG-1 (MP3), covers 20kHz bandwidth of human hearing with room for LP ripple. */
-		48000,		/** Standard rate used by "professional" film and audio guys: mixing console, digital recorders, etc */
-		88200,		/** Used for recording equipment intended for CDs */
-		96000,		/** DVD audio, high-def audio, 2x the 48khz "professional" sample rate. */
-		176400,		/** Rate used by HDCD recorders. */
-		192000,		/** HD-DVD and blue ray audio, 4x the 48khz "professional" sample rate. */
-	};
-	const uint32 IUnrealAudioDeviceModule::MaxPossibleFrameRates = 11;
-
-	IUnrealAudioDeviceModule::IUnrealAudioDeviceModule()
-	{
-	}
-
-	IUnrealAudioDeviceModule::~IUnrealAudioDeviceModule()
-	{
-	}
-
-	bool IUnrealAudioDeviceModule::CreateStream(const FCreateStreamParams& Params)
-	{
-		checkf(StreamInfo.State == EStreamState::CLOSED, TEXT("Stream state can't be open if creating a new one."));
-		checkf(Params.OutputDeviceIndex != INDEX_NONE, TEXT("Input or output stream params need to be set."));
-
-		bool bSuccess = false;
-
-		Reset();
-
-		if (OpenDevice(Params))
-		{
-			StreamInfo.State = EStreamState::STOPPED;
-			StreamInfo.CallbackFunction = Params.CallbackFunction;
-			StreamInfo.UserData = Params.UserData;
-			StreamInfo.BlockSize = Params.CallbackBlockSize;
-			StreamInfo.StreamDelta = (double)StreamInfo.BlockSize / StreamInfo.FrameRate;
-			bSuccess = true;
-		}
-
-		return bSuccess;
-	}
-
-	void IUnrealAudioDeviceModule::UpdateStreamTimeTick()
-	{
-		StreamInfo.StreamTime += StreamInfo.StreamDelta;
-	}
-
-	void IUnrealAudioDeviceModule::OnDeviceError(const EDeviceError::Type ErrorType, const FString& ErrorDetails, const FString& FileName, int32 LineNumber) const
-	{
-		UE_LOG(LogUnrealAudioDevice, Error, TEXT("Audio Device Error: (%s) : %s (%s::%d)"), EDeviceError::ToString(ErrorType), *ErrorDetails, *FileName, LineNumber);
-	}
-
-	void IUnrealAudioDeviceModule::Reset()
-	{
-		StreamInfo.Initialize();
-	}
-
 
 	/** 
 	* FUnrealAudioDeviceDummy
@@ -164,6 +100,24 @@ namespace UAudio
 		return new FUnrealAudioDeviceDummy();
 	}
 
+// 	void OnDeviceError(const EDeviceError::Type Error, const FString& ErrorDetails, const FString& FileName, int32 LineNumber)
+// 	{
+// 		UE_LOG(LogUnrealAudioDevice, Error, TEXT("Audio Device Error: (%s) : %s (%s::%d)"), EDeviceError::ToString(Error), *ErrorDetails, *FileName, LineNumber);
+// 	}
+
+// 	uint32 GetNumBytesForFormat(EStreamFormat::Type Format)
+// 	{
+// 		switch (Format)
+// 		{
+// 			case EStreamFormat::FLT:		return 4;
+// 			case EStreamFormat::DBL:		return 8;
+// 			case EStreamFormat::INT_16:		return 2;
+// 			case EStreamFormat::INT_24:		return 3;
+// 			case EStreamFormat::INT_32:		return 4;
+// 			default:						return 0;
+// 		}
+// 	}
+
 }
 
 #endif // #if ENABLE_UNREAL_AUDIO
diff --git a/Engine/Source/Runtime/UnrealAudio/Private/UnrealAudioDeviceFormat.cpp b/Engine/Source/Runtime/UnrealAudio/Private/UnrealAudioDeviceFormat.cpp
--- a/Engine/Source/Runtime/UnrealAudio/Private/UnrealAudioDeviceFormat.cpp
+++ b/Engine/Source/Runtime/UnrealAudio/Private/UnrealAudioDeviceFormat.cpp
@@ -14,19 +14,6 @@ namespace UAudio
 	 * so we need to support converting byte buffers to and from various formats.
 	 */
 
-	uint32 IUnrealAudioDeviceModule::GetNumBytesForFormat(EStreamFormat::Type Format) const
-	{
-		switch (Format)
-		{
-			case EStreamFormat::FLT:		return 4;
-			case EStreamFormat::DBL:		return 8;
-			case EStreamFormat::INT_16:		return 2;
-			case EStreamFormat::INT_24:		return 3;
-			case EStreamFormat::INT_32:		return 4;
-			default:						return 0;
-		}
-	}
-
 	void IUnrealAudioDeviceModule::SetupBufferFormatConvertInfo()
 	{
 		FBufferFormatConvertInfo& ConvertInfo = StreamInfo.DeviceInfo.BufferFormatConvertInfo;
diff --git a/Engine/Source/Runtime/UnrealAudio/Public/UnrealAudioDeviceModule.h b/Engine/Source/Runtime/UnrealAudio/Public/UnrealAudioDeviceModule.h
--- a/Engine/Source/Runtime/UnrealAudio/Public/UnrealAudioDeviceModule.h
+++ b/Engine/Source/Runtime/UnrealAudio/Public/UnrealAudioDeviceModule.h
@@ -10,9 +10,9 @@
 
 DECLARE_LOG_CATEGORY_EXTERN(LogUnrealAudioDevice, Log, All);
 
-#define UA_DEVICE_PLATFORM_ERROR(INFO)	(OnDeviceError(EDeviceError::PLATFORM, INFO, FString(__FILE__), __LINE__))
-#define AU_DEVICE_PARAM_ERROR(INFO)		(OnDeviceError(EDeviceError::INVALID_PARAMETER, INFO, FString(__FILE__), __LINE__))
-#define AU_DEVICE_WARNING(INFO)			(OnDeviceError(EDeviceError::WARNING, INFO, FString(__FILE__), __LINE__))
+#define UA_DEVICE_PLATFORM_ERROR(INFO)	(UAudio::OnDeviceError(EDeviceError::PLATFORM, INFO, FString(__FILE__), __LINE__))
+#define AU_DEVICE_PARAM_ERROR(INFO)		(UAudio::OnDeviceError(EDeviceError::INVALID_PARAMETER, INFO, FString(__FILE__), __LINE__))
+#define AU_DEVICE_WARNING(INFO)			(UAudio::OnDeviceError(EDeviceError::WARNING, INFO, FString(__FILE__), __LINE__))
 
 
 namespace UAudio
@@ -95,7 +95,7 @@ namespace UAudio
 	{
 		enum Type
 		{
-			CLOSED,				/** The stream is closed */
+			SHUTDOWN,			/** The stream is shutdown */
 			STOPPED,			/** The stream was running and is now stopped */
 			RUNNING,			/** The stream is running (currently streaming callbacks) */
 			STOPPING,			/** The stream is currently stopping*/
@@ -110,7 +110,7 @@ namespace UAudio
 	{
 		enum Flag
 		{
-			OUTPUT_UNDERFLOW	= (1 << 1)		/** The output stream has too little data */
+			OUTPUT_OVERFLOW	= (1 << 1)		/** The output stream has too little or too much data */
 		};
 	}
 	typedef uint8 StreamStatus;
@@ -141,6 +141,7 @@ namespace UAudio
 			TOP_BACK_LEFT,
 			TOP_BACK_CENTER,
 			TOP_BACK_RIGHT,
+			UNUSED,
 			SPEAKER_TYPE_COUNT
 		};
 
@@ -167,12 +168,32 @@ namespace UAudio
 				case TOP_BACK_LEFT:				return TEXT("TOP_BACK_LEFT");
 				case TOP_BACK_CENTER:			return TEXT("TOP_BACK_CENTER");
 				case TOP_BACK_RIGHT:			return TEXT("TOP_BACK_RIGHT");
+				case UNUSED:					return TEXT("UNUSED");
 				default:						return TEXT("UKNOWN");
 			}
 			return TEXT("");
 		}
 	};
 
+	/**
+	Static array specifying which frame rates to check for support on audio devices.
+	http://en.wikipedia.org/wiki/Sampling_%28signal_processing%29
+	*/
+	const uint32 PossibleFrameRates[] =
+	{
+		8000,		/** Used for telephony, walkie talkies, ok for human speech. */
+		11025,		/** Quarter sample rate of CD's, low-quality PCM */
+		16000,		/** Rate used for VOIP (which is why VOIP sounds slightly better than normal phones), wide-band extension over normal telephony */
+		22050,		/** Half CD sample rate, low quality PCM */
+		32000,		/** MiniDV, digital FM radio, decent wireless microphones */
+		44100,		/** CD's, MPEG-1 (MP3), covers 20kHz bandwidth of human hearing with room for LP ripple. */
+		48000,		/** Standard rate used by "professional" film and audio guys: mixing console, digital recorders, etc */
+		88200,		/** Used for recording equipment intended for CDs */
+		96000,		/** DVD audio, high-def audio, 2x the 48khz "professional" sample rate. */
+		176400,		/** Rate used by HDCD recorders. */
+		192000,		/** HD-DVD and blue ray audio, 4x the 48khz "professional" sample rate. */
+	};
+	const uint32 MaxPossibleFrameRates = 11;
 
 	/**
 	* FCallbackInfo
@@ -187,12 +208,15 @@ namespace UAudio
 		/** The number of buffer output frames */
 		uint32 NumFrames;
 
-		/** Array of outputs speakers */
-		TArray<ESpeaker::Type> OutputSpeakers;
-
 		/** The number of channels in output. */
 		int32 NumChannels;
 
+		/** The number of total samples (NumFrames*NumChannels). */
+		int32 NumSamples;
+
+		/** Array of outputs speakers */
+		TArray<ESpeaker::Type> OutputSpeakers;
+
 		/** The current status flags of the input and output buffers */
 		StreamStatus StatusFlags;
 
@@ -229,22 +253,29 @@ namespace UAudio
 		/** The frame rate of the device. */
 		uint32 FrameRate;
 
+		/** The possible frame rates of the device. */
+		TArray<uint32> PossibleFrameRates;
+
 		/** The data format of the device (e.g. float). */
 		EStreamFormat::Type StreamFormat;
 
 		/** What speakers this device supports (if output device) */
 		TArray<ESpeaker::Type> Speakers;
 
+		/** Device latency (if available) */
+		uint32 Latency;
+
 		/** Whether or not it is the OS default device for the type. */
 		bool bIsSystemDefault;
 
 		FDeviceInfo()
 			: FriendlyName(TEXT("Uknown"))
 			, NumChannels(0)
-			, FrameRate(0)
 			, StreamFormat(EStreamFormat::UNKNOWN)
+			, Latency(0)
 			, bIsSystemDefault(false)
-		{}
+		{
+		}
 	};
 
 	/**
@@ -278,29 +309,22 @@ namespace UAudio
 		{}
 	};
 
-	/**
-	* IDeviceErrorListener
-	* An optional class that defines an error listener interface. Used for handling or logging any audio device errors which may have occurred.
-	*/
-	class UNREALAUDIO_API IDeviceErrorListener
-	{
-	public:
-		virtual void OnDeviceError(const EDeviceError::Type Error, const FString& ErrorDetails, const FString& FileName, int32 LineNumber) const = 0;
-	};
-
 	/**
 	* IUnrealAudioDeviceModule
 	* Main audio device module, needs to be implemented for every specific platform API.
 	*/
-	class UNREALAUDIO_API IUnrealAudioDeviceModule : public IModuleInterface,
-													 public IDeviceErrorListener
+	class UNREALAUDIO_API IUnrealAudioDeviceModule : public IModuleInterface
 	{
 	public:
 		/** Constructor */
-		IUnrealAudioDeviceModule();
+		IUnrealAudioDeviceModule()
+		{
+		}
 
 		/** Virtual destructor */
-		virtual ~IUnrealAudioDeviceModule();
+		virtual ~IUnrealAudioDeviceModule()
+		{
+		}
 
 		/** 
 		* Initializes the audio device module. 
@@ -390,16 +414,27 @@ namespace UAudio
 		* @param Params The parameters needed to create an audio stream.
 		* @return true if succeeded.
 		*/
-		bool CreateStream(const FCreateStreamParams& Params);
+		bool CreateStream(const FCreateStreamParams& Params)
+		{
+			checkf(StreamInfo.State == EStreamState::SHUTDOWN, TEXT("Stream state can't be open if creating a new one."));
+			checkf(Params.OutputDeviceIndex != INDEX_NONE, TEXT("Input or output stream params need to be set."));
 
-		/**
-		* Function called when an error occurs in the device code.
-		* @param Error The error type that occurred.
-		* @param ErrorDetails A string of specific details about the error.
-		* @param FileName The file that the error occurred in.
-		* @param LineNumber The line number that the error occurred in.
-		*/
-		void OnDeviceError(const EDeviceError::Type Error, const FString& ErrorDetails, const FString& FileName, int32 LineNumber) const override;
+			bool bSuccess = false;
+
+			Reset();
+
+			if (OpenDevice(Params))
+			{
+				StreamInfo.State = EStreamState::STOPPED;
+				StreamInfo.CallbackFunction = Params.CallbackFunction;
+				StreamInfo.UserData = Params.UserData;
+				StreamInfo.BlockSize = Params.CallbackBlockSize;
+				StreamInfo.StreamDelta = (double)StreamInfo.BlockSize / StreamInfo.FrameRate;
+				bSuccess = true;
+			}
+
+			return bSuccess;
+		}
 
 	protected:
 
@@ -443,7 +478,7 @@ namespace UAudio
 			uint32 NumChannels;
 
 			/** The reported latency of this device. */
-			uint64 Latency;
+			uint32 Latency;
 
 			/** The native framerate of this device. */
 			uint32 FrameRate;
@@ -539,7 +574,7 @@ namespace UAudio
 			void Initialize()
 			{
 				FrameRate = 0;
-				State = EStreamState::CLOSED;
+				State = EStreamState::SHUTDOWN;
 				StreamTime = 0.0;
 				StreamDelta = 0.0;
 
@@ -557,10 +592,10 @@ namespace UAudio
 		virtual bool OpenDevice(const FCreateStreamParams& Params) = 0;
 
 		/** Called before opening up new streams. */
-		void Reset();
-
-		/** Helper function to get number of bytes for a given stream format. */
-		uint32 GetNumBytesForFormat(EStreamFormat::Type Format) const;
+		void Reset()
+		{
+			StreamInfo.Initialize();
+		}
 
 		/** Sets up any convert information for given stream type (figures out to/from convert format and channel formats) */
 		void SetupBufferFormatConvertInfo();
@@ -569,7 +604,10 @@ namespace UAudio
 		bool ConvertBufferFormat(TArray<uint8>& OutputBuffer, TArray<uint8>& InputBuffer);
 
 		/** Updates the sample-accurate stream time value. */
-		void UpdateStreamTimeTick();
+		FORCEINLINE void UpdateStreamTimeTick()
+		{
+			StreamInfo.StreamTime += StreamInfo.StreamDelta;
+		}
 
 		/** Templated function to convert to float from any other format type. */
 		template <typename FromType, typename ToType>
@@ -593,17 +631,42 @@ namespace UAudio
 
 	protected: // Protected Data
 
-		/** An array of possible frame rates that can exist with any audio device. */
-		static const uint32 PossibleFrameRates[];
-		static const uint32 MaxPossibleFrameRates;
-
 		/** Device stream info struct */
 		FStreamInfo StreamInfo;
 	};
 
 	// Exported Functions
+
+	/** Creates a dummy audio device */
 	UNREALAUDIO_API IUnrealAudioDeviceModule* CreateDummyDeviceModule();
-}
+
+	/**
+	* Function called when an error occurs in the device code.
+	* @param Error The error type that occurred.
+	* @param ErrorDetails A string of specific details about the error.
+	* @param FileName The file that the error occurred in.
+	* @param LineNumber The line number that the error occurred in.
+	*/
+	static inline void OnDeviceError(const EDeviceError::Type Error, const FString& ErrorDetails, const FString& FileName, int32 LineNumber)
+	{
+		UE_LOG(LogUnrealAudioDevice, Error, TEXT("Audio Device Error: (%s) : %s (%s::%d)"), EDeviceError::ToString(Error), *ErrorDetails, *FileName, LineNumber);
+	}
+
+	/** Helper function to get number of bytes for a given stream format. */
+	static inline uint32 GetNumBytesForFormat(EStreamFormat::Type Format)
+	{
+		switch (Format)
+		{
+			case EStreamFormat::FLT:		return 4;
+			case EStreamFormat::DBL:		return 8;
+			case EStreamFormat::INT_16:		return 2;
+			case EStreamFormat::INT_24:		return 3;
+			case EStreamFormat::INT_32:		return 4;
+			default:						return 0;
+		}
+	}
+
+} // namespace UAudio
 
 #endif // #if ENABLE_UNREAL_AUDIO
 
diff --git a/Engine/Source/Runtime/UnrealAudio/Public/UnrealAudioModule.h b/Engine/Source/Runtime/UnrealAudio/Public/UnrealAudioModule.h
--- a/Engine/Source/Runtime/UnrealAudio/Public/UnrealAudioModule.h
+++ b/Engine/Source/Runtime/UnrealAudio/Public/UnrealAudioModule.h
@@ -4,8 +4,8 @@
 
 #include "ModuleInterface.h"
 
-// Only enable unreal audio on windows
-#if PLATFORM_WINDOWS
+// Only enable unreal audio on windows or mac
+#if PLATFORM_WINDOWS || PLATFORM_MAC
 #define ENABLE_UNREAL_AUDIO 1
 #else
 #define ENABLE_UNREAL_AUDIO 0
diff --git a/Engine/Source/Runtime/Windows/UnrealAudioWasapi/UnrealAudioWasapi.Build.cs b/Engine/Source/Runtime/Windows/UnrealAudioWasapi/UnrealAudioWasapi.Build.cs
--- a/Engine/Source/Runtime/Windows/UnrealAudioWasapi/UnrealAudioWasapi.Build.cs
+++ b/Engine/Source/Runtime/Windows/UnrealAudioWasapi/UnrealAudioWasapi.Build.cs
@@ -6,11 +6,9 @@ public class UnrealAudioWasapi : ModuleRules
 {
 	public UnrealAudioWasapi(TargetInfo Target)
 	{
-		PrivateIncludePaths.AddRange(
-			new string[] {
-				"Runtime/UnrealAudio/Private",
-			}
-)		;
+		PrivateIncludePaths.Add("Runtime/UnrealAudio/Private");
+		PublicIncludePaths.Add("Runtime/UnrealAudio/Public");
+
 
 		PrivateIncludePathModuleNames.Add("TargetPlatform");
 
@@ -21,10 +19,7 @@ public class UnrealAudioWasapi : ModuleRules
 			}
 		);
 
-		PublicDependencyModuleNames.AddRange(
-			new string[] {
-				"UnrealAudio"
-			}
-		);
+		PublicDependencyModuleNames.Add("UnrealAudio");
+
 	}
 }
diff --git a/Engine/Source/Runtime/Windows/UnrealAudioXAudio2/Private/UnrealAudioDeviceXAudio2.cpp b/Engine/Source/Runtime/Windows/UnrealAudioXAudio2/Private/UnrealAudioDeviceXAudio2.cpp
--- a/Engine/Source/Runtime/Windows/UnrealAudioXAudio2/Private/UnrealAudioDeviceXAudio2.cpp
+++ b/Engine/Source/Runtime/Windows/UnrealAudioXAudio2/Private/UnrealAudioDeviceXAudio2.cpp
@@ -45,19 +45,19 @@ static const TCHAR* GetXAudio2Error(HRESULT Result)
 
 #define CLEANUP_ON_FAIL(Result)						\
 	if (FAILED(Result))								\
-		{												\
+	{												\
 		const TCHAR* Err = GetXAudio2Error(Result);	\
 		UA_DEVICE_PLATFORM_ERROR(Err);				\
 		goto Cleanup;								\
-		}
+	}
 
 #define RETURN_FALSE_ON_FAIL(Result)				\
 	if (FAILED(Result))								\
-		{												\
+	{												\
 		const TCHAR* Err = GetXAudio2Error(Result);	\
 		UA_DEVICE_PLATFORM_ERROR(Err);				\
 		return false;								\
-		}
+	}
 
 namespace UAudio
 {
@@ -253,6 +253,11 @@ namespace UAudio
 
 	bool FUnrealAudioXAudio2::Initialize()
 	{
+		if (bInitialized)
+		{
+			return false;
+		}
+
 		bComInitialized = FWindowsPlatformMisc::CoInitialize();
 		HRESULT Result = XAudio2Create(&XAudio2Info.XAudio2System);
 		RETURN_FALSE_ON_FAIL(Result);
@@ -409,7 +414,7 @@ namespace UAudio
 
 	bool FUnrealAudioXAudio2::StopStream()
 	{
-		if (!bInitialized || StreamInfo.State == EStreamState::CLOSED || StreamInfo.State == EStreamState::STOPPED || StreamInfo.State == EStreamState::STOPPING)
+		if (!bInitialized || StreamInfo.State == EStreamState::SHUTDOWN || StreamInfo.State == EStreamState::STOPPED || StreamInfo.State == EStreamState::STOPPING)
 		{
 			return false;
 		}
@@ -433,7 +438,7 @@ namespace UAudio
 
 	bool FUnrealAudioXAudio2::ShutdownStream()
 	{
-		if (!bInitialized || StreamInfo.State == EStreamState::CLOSED)
+		if (!bInitialized || StreamInfo.State == EStreamState::SHUTDOWN)
 		{
 			return false;
 		}
@@ -453,17 +458,19 @@ namespace UAudio
 		CloseHandle(XAudio2Info.OutputBufferEndEvent);
 		XAudio2Info.OutputBufferEndEvent = nullptr;
 
-		StreamInfo.State = EStreamState::CLOSED;
+		StreamInfo.State = EStreamState::SHUTDOWN;
 		return true;
 	}
 
 	bool FUnrealAudioXAudio2::GetLatency(uint32& OutputDeviceLatency) const
 	{
+		OutputDeviceLatency = StreamInfo.DeviceInfo.Latency;
 		return true;
 	}
 
 	bool FUnrealAudioXAudio2::GetFrameRate(uint32& OutFrameRate) const
 	{
+		OutFrameRate = StreamInfo.FrameRate;
 		return true;
 	}
 
diff --git a/Engine/Source/Runtime/Windows/UnrealAudioXAudio2/UnrealAudioXAudio2.Build.cs b/Engine/Source/Runtime/Windows/UnrealAudioXAudio2/UnrealAudioXAudio2.Build.cs
--- a/Engine/Source/Runtime/Windows/UnrealAudioXAudio2/UnrealAudioXAudio2.Build.cs
+++ b/Engine/Source/Runtime/Windows/UnrealAudioXAudio2/UnrealAudioXAudio2.Build.cs
@@ -6,11 +6,8 @@ public class UnrealAudioXAudio2 : ModuleRules
 {
 	public UnrealAudioXAudio2(TargetInfo Target)
 	{
-		PrivateIncludePaths.AddRange(
-			new string[] {
-				"Runtime/UnrealAudio/Private",
-			}
-)		;
+		PrivateIncludePaths.Add("Runtime/UnrealAudio/Private");
+		PublicIncludePaths.Add("Runtime/UnrealAudio/Public");
 
 		PrivateIncludePathModuleNames.Add("TargetPlatform");
 
@@ -20,12 +17,8 @@ public class UnrealAudioXAudio2 : ModuleRules
 				"CoreUObject",
 			}
 		);
-
-		PublicDependencyModuleNames.AddRange(
-			new string[] {
-				"UnrealAudio"
-			}
-		);
+		
+		PublicDependencyModuleNames.Add("UnrealAudio");
 
 		AddThirdPartyPrivateStaticDependencies(Target, 
 			"DX11Audio"
