Fix shadow variables

diff --git a/Engine/Source/Editor/PropertyEditor/Private/PropertyCustomizationHelpers.cpp b/Engine/Source/Editor/PropertyEditor/Private/PropertyCustomizationHelpers.cpp
--- a/Engine/Source/Editor/PropertyEditor/Private/PropertyCustomizationHelpers.cpp
+++ b/Engine/Source/Editor/PropertyEditor/Private/PropertyCustomizationHelpers.cpp
@@ -1066,9 +1066,9 @@ TSharedRef<SWidget> PropertyCustomizationHelpers::MakeTextLocalizationButton(con
 		{
 		}
 
-		void Construct(const FArguments& InArgs, const TSharedRef<IPropertyHandle>& InPropertyHandle)
+		void Construct(const FArguments& InArgs, const TSharedRef<IPropertyHandle>& InPropHandle)
 		{
-			PropertyHandle = InPropertyHandle;
+			PropertyHandle = InPropHandle;
 
 			FText DisplayText;
 			if (PropertyHandle->GetValueAsDisplayText(DisplayText) == FPropertyAccess::Success)
diff --git a/Engine/Source/Runtime/CoreUObject/Private/Serialization/BulkData.cpp b/Engine/Source/Runtime/CoreUObject/Private/Serialization/BulkData.cpp
--- a/Engine/Source/Runtime/CoreUObject/Private/Serialization/BulkData.cpp
+++ b/Engine/Source/Runtime/CoreUObject/Private/Serialization/BulkData.cpp
@@ -1450,7 +1450,7 @@ void FFloatBulkData::SerializeElement( FArchive& Ar, void* Data, int32 ElementIn
 	Ar << FloatData;
 }
 
-void FFormatContainer::Serialize(FArchive& Ar, UObject* Owner, const TArray<FName>* FormatsToSave, bool bSingleUse, uint32 Alignment)
+void FFormatContainer::Serialize(FArchive& Ar, UObject* Owner, const TArray<FName>* FormatsToSave, bool bSingleUse, uint32 InAlignment)
 {
 	if (Ar.IsLoading())
 	{
@@ -1461,7 +1461,7 @@ void FFormatContainer::Serialize(FArchive& Ar, UObject* Owner, const TArray<FNam
 			FName Name;
 			Ar << Name;
 			FByteBulkData& Bulk = GetFormat(Name);
-			Bulk.SetBulkDataAlignment(Alignment);
+			Bulk.SetBulkDataAlignment(InAlignment);
 			Bulk.Serialize(Ar, Owner);
 		}
 	}
diff --git a/Engine/Source/Runtime/CoreUObject/Private/UObject/PropertyMap.cpp b/Engine/Source/Runtime/CoreUObject/Private/UObject/PropertyMap.cpp
--- a/Engine/Source/Runtime/CoreUObject/Private/UObject/PropertyMap.cpp
+++ b/Engine/Source/Runtime/CoreUObject/Private/UObject/PropertyMap.cpp
@@ -658,7 +658,6 @@ void UMapProperty::CopyValuesInternal(void* Dest, void const* Src, int32 Count)
 		return;
 	}
 
-	int32 DestIndex = 0;
 	for (int32 SrcIndex = 0; Num; ++SrcIndex)
 	{
 		if (SrcMapHelper.IsValidIndex(SrcIndex))
diff --git a/Engine/Source/Runtime/Mac/UnrealAudioCoreAudio/Private/UnrealAudioCoreAudio.cpp b/Engine/Source/Runtime/Mac/UnrealAudioCoreAudio/Private/UnrealAudioCoreAudio.cpp
--- a/Engine/Source/Runtime/Mac/UnrealAudioCoreAudio/Private/UnrealAudioCoreAudio.cpp
+++ b/Engine/Source/Runtime/Mac/UnrealAudioCoreAudio/Private/UnrealAudioCoreAudio.cpp
@@ -690,7 +690,7 @@ namespace UAudio
 		return bSuccess;
 	}
 
-	bool FUnrealAudioCoreAudio::GetDeviceFrameRates(AudioDeviceID DeviceID, TArray<uint32>& PossibleFrameRates, uint32& DefaultFrameRate)
+	bool FUnrealAudioCoreAudio::GetDeviceFrameRates(AudioDeviceID DeviceID, TArray<uint32>& InPossibleFrameRates, uint32& DefaultFrameRate)
 	{
 		AudioObjectPropertyAddress Property = NEW_OUTPUT_PROPERTY(kAudioDevicePropertyAvailableNominalSampleRates);
 
@@ -717,9 +717,9 @@ namespace UAudio
 			if (SampleRate.mMinimum == SampleRate.mMaximum)
 			{
 				uint32 FrameRate = (uint32)SampleRate.mMinimum;
-				if (!PossibleFrameRates.Contains(FrameRate))
+				if (!InPossibleFrameRates.Contains(FrameRate))
 				{
-					PossibleFrameRates.Add(FrameRate);
+					InPossibleFrameRates.Add(FrameRate);
 				}
 			}
 			else
@@ -741,19 +741,19 @@ namespace UAudio
 		{
 			for (uint32 i = 0; i < MaxPossibleFrameRates; ++i)
 			{
-				uint32 PossibleFrameRate = PossibleFrameRates[i];
+				uint32 PossibleFrameRate = InPossibleFrameRates[i];
 				if (PossibleFrameRate >= LargestMinSampleRateRange && PossibleFrameRate <= SmallestMaxSampleRateRange)
 				{
-					PossibleFrameRates.Add(PossibleFrameRate);
+					InPossibleFrameRates.Add(PossibleFrameRate);
 				}
 			}
 		}
 		// Picking a good default framerate from the supported framerates
-		if (PossibleFrameRates.Contains(48000))
+		if (InPossibleFrameRates.Contains(48000))
 		{
 			DefaultFrameRate = 48000;
 		}
-		else if (PossibleFrameRates.Contains(44100))
+		else if (InPossibleFrameRates.Contains(44100))
 		{
 			DefaultFrameRate = 44100;
 		}
diff --git a/Engine/Source/Runtime/Online/OnlineSubsystem/Public/Interfaces/OnlinePartyInterface.h b/Engine/Source/Runtime/Online/OnlineSubsystem/Public/Interfaces/OnlinePartyInterface.h
--- a/Engine/Source/Runtime/Online/OnlineSubsystem/Public/Interfaces/OnlinePartyInterface.h
+++ b/Engine/Source/Runtime/Online/OnlineSubsystem/Public/Interfaces/OnlinePartyInterface.h
@@ -132,9 +132,9 @@ namespace EPartyState
 class FOnlinePartyInfo
 {
 public:
-	FOnlinePartyInfo(const TSharedPtr<FOnlinePartyId>& PartyId = TSharedPtr<FOnlinePartyId>())
+	FOnlinePartyInfo(const TSharedPtr<FOnlinePartyId>& InPartyId = TSharedPtr<FOnlinePartyId>())
 		: State(EPartyState::None)
-		, PartyId(PartyId)
+		, PartyId(InPartyId )
 	{}
 
 	virtual ~FOnlinePartyInfo() 
diff --git a/Engine/Source/Runtime/Online/OnlineSubsystemUtils/Private/Tests/TestSessionInterface.cpp b/Engine/Source/Runtime/Online/OnlineSubsystemUtils/Private/Tests/TestSessionInterface.cpp
--- a/Engine/Source/Runtime/Online/OnlineSubsystemUtils/Private/Tests/TestSessionInterface.cpp
+++ b/Engine/Source/Runtime/Online/OnlineSubsystemUtils/Private/Tests/TestSessionInterface.cpp
@@ -180,7 +180,7 @@ void FTestSessionInterface::OnReadFriendsListComplete(int32 LocalUserNum, bool b
 	}
 }
 
-void FTestSessionInterface::OnSessionUserInviteAccepted(bool bWasSuccessful, int32 ControllerId, TSharedPtr<FUniqueNetId> UserId, const FOnlineSessionSearchResult& SearchResult)
+void FTestSessionInterface::OnSessionUserInviteAccepted(bool bWasSuccessful, int32 ControllerId, TSharedPtr<FUniqueNetId> InUserId, const FOnlineSessionSearchResult& SearchResult)
 {
 	UE_LOG(LogOnline, Verbose, TEXT("OnSessionInviteAccepted ControllerId: %d bSuccess: %d"), ControllerId, bWasSuccessful);
 	// Don't clear invite accept delegate
diff --git a/Engine/Source/Runtime/Renderer/Private/ShadowRendering.cpp b/Engine/Source/Runtime/Renderer/Private/ShadowRendering.cpp
--- a/Engine/Source/Runtime/Renderer/Private/ShadowRendering.cpp
+++ b/Engine/Source/Runtime/Renderer/Private/ShadowRendering.cpp
@@ -1002,7 +1002,7 @@ void FShadowDepthDrawingPolicyFactory::AddStaticMesh(FScene* Scene,FStaticMesh*
 		const EBlendMode BlendMode = Material->GetBlendMode();
 		const EMaterialShadingModel ShadingModel = Material->GetShadingModel();
 
-		const bool bLightPropagationVolume = UseLightPropagationVolumeRT(Scene->GetFeatureLevel());
+		const bool bLightPropagationVolume = UseLightPropagationVolumeRT(FeatureLevel);
 		const bool bTwoSided  = Material->IsTwoSided() || StaticMesh->PrimitiveSceneInfo->Proxy->CastsShadowAsTwoSided();
 		const bool bLitOpaque = !IsTranslucentBlendMode(BlendMode) && ShadingModel != MSM_Unlit;
 		if(bLightPropagationVolume && (bLitOpaque || Material->ShouldInjectEmissiveIntoLPV()))
@@ -1029,8 +1029,7 @@ void FShadowDepthDrawingPolicyFactory::AddStaticMesh(FScene* Scene,FStaticMesh*
 		}
 		if ( bLitOpaque )
 		{
-			const auto FeatureLevel = Scene->GetFeatureLevel();
-			OverrideWithDefaultMaterialForShadowDepth(MaterialRenderProxy, Material, false, Scene->GetFeatureLevel()); 
+			OverrideWithDefaultMaterialForShadowDepth(MaterialRenderProxy, Material, false, FeatureLevel); 
 
 			// Add the static mesh to the shadow's subject draw list.
 			Scene->WholeSceneShadowDepthDrawList.AddMesh(
diff --git a/Engine/Source/Runtime/UnrealAudio/Private/Tests/UnrealAudioTestGenerators.cpp b/Engine/Source/Runtime/UnrealAudio/Private/Tests/UnrealAudioTestGenerators.cpp
--- a/Engine/Source/Runtime/UnrealAudio/Private/Tests/UnrealAudioTestGenerators.cpp
+++ b/Engine/Source/Runtime/UnrealAudio/Private/Tests/UnrealAudioTestGenerators.cpp
@@ -86,8 +86,8 @@ namespace Test
 	/* FTimer																*/
 	/************************************************************************/
 
-	FTimer::FTimer(double TotalTime)
-		: TotalTime(TotalTime)
+	FTimer::FTimer(double InTotalTime)
+		: TotalTime(InTotalTime)
 		, StartTime(0.0f)
 		, CurrentTime(0.0f)
 		, LastTime(0.0f)
@@ -860,9 +860,9 @@ namespace Test
 			FilterLFO.SetFrequency(FMath::FRandRange(0.1f, 20.0f));
 			FilterLFO.SetScaleAdd(1000.0f, 2500.0f);
 			LowPass.SetParams(1500.0f, 1.0f);
-			for (int32 ModIndex = 0; ModIndex < Mods.Num(); ++ModIndex)
+			for (OscData& Mod : Mods)
 			{
-				Mods[ModIndex].Reset();
+				Mod.Reset();
 			}
 		}
 	}
