Merging CL 2590157 - Workaround for CoreAudio linking issue with OS X 10.11 SDK

diff --git a/Engine/Source/Runtime/Mac/CoreAudio/CoreAudio.Build.cs b/Engine/Source/Runtime/Mac/CoreAudio/CoreAudio.Build.cs
--- a/Engine/Source/Runtime/Mac/CoreAudio/CoreAudio.Build.cs
+++ b/Engine/Source/Runtime/Mac/CoreAudio/CoreAudio.Build.cs
@@ -22,7 +22,7 @@ public class CoreAudio : ModuleRules
 			"VorbisFile"
 			);
 
-		PublicFrameworks.AddRange(new string[] { "CoreAudio", "AudioToolbox", "AudioUnit" });
+		PublicFrameworks.AddRange(new string[] { "CoreAudio", "AudioUnit", "AudioToolbox" });
 
 		AdditionalBundleResources.Add(new UEBuildBundleResource("ThirdParty/Mac/RadioEffectUnit/RadioEffectUnit.component"));
 	}
