BLOCKING ISSUE UE-17563 Fixing core audio crash

- QAGame's Test map TM-MultipleSounds allowed me to finally repro the issue after a few minutes. Reducing the number of voices in CoreAudio caused the issue to happen immediately.
- Crash was due to an invalid audio unit index (-1) when the number of active audio units exceeds total number of available audio units
- Fixed by properly failing to play a sound if no audio unit/voice indices are available rather than trying to construct an AUGraph with the invalid index.
- Added more asserts to make sure things are what they should be.
- Changed the MAX_AUDIOCHANNELS #define used in Mac-code so it can be different than other platforms. Other platforms use this as an upper limit but don'''t allocate resources according to it, Mac however, allocates resources and allows voices to play up to the max number. Rather than fix that right now, we''ll hard-code that upper limit on mac to be on-par with PC
- Went ahead and reduced the number of max voices on Mac since PC only uses 32 and this was set to be 64 (for mono/stereo files) and 16 for multichannel files for a total of 80 max voices. Mac is already perf-heavy on the main thread, this will help to decrease perf due to audio and will be more in-line with PC audio.

diff --git a/Engine/Source/Runtime/Mac/CoreAudio/Private/CoreAudioDevice.cpp b/Engine/Source/Runtime/Mac/CoreAudio/Private/CoreAudioDevice.cpp
--- a/Engine/Source/Runtime/Mac/CoreAudio/Private/CoreAudioDevice.cpp
+++ b/Engine/Source/Runtime/Mac/CoreAudio/Private/CoreAudioDevice.cpp
@@ -50,12 +50,12 @@ bool FCoreAudioDevice::InitializeHardware()
 
 	InverseTransform = FMatrix::Identity;
 
-	for( SInt32 Index = 0; Index < MAX_AUDIOCHANNELS; ++Index )
+	for( SInt32 Index = 0; Index < CORE_AUDIO_MAX_CHANNELS; ++Index )
 	{
 		Mixer3DInputStatus[ Index ] = false;
 	}
 
-	for( SInt32 Index = 0; Index < MAX_MULTICHANNEL_AUDIOCHANNELS; ++Index )
+	for( SInt32 Index = 0; Index < CORE_AUDIO_MAX_MULTICHANNEL_AUDIOCHANNELS; ++Index )
 	{
 		MatrixMixerInputStatus[ Index ] = false;
 	}
@@ -144,7 +144,7 @@ bool FCoreAudioDevice::InitializeHardware()
 		Status = AUGraphNodeInfo( AudioUnitGraph, MatrixMixerNode, NULL, &MatrixMixerUnit );
 		
 		// Set number of buses for input
-		uint32 NumBuses = MAX_MULTICHANNEL_AUDIOCHANNELS;
+		uint32 NumBuses = CORE_AUDIO_MAX_MULTICHANNEL_AUDIOCHANNELS;
 		Size = sizeof( NumBuses );
 		
 		Status = AudioUnitSetProperty( MatrixMixerUnit, kAudioUnitProperty_ElementCount, kAudioUnitScope_Input, 0, &NumBuses, Size );
@@ -172,7 +172,7 @@ bool FCoreAudioDevice::InitializeHardware()
 		MatrixMixerInputFormat.mBytesPerPacket = MatrixMixerInputFormat.mBytesPerFrame;
 		MatrixMixerInputFormat.mFormatFlags |= kAudioFormatFlagIsNonInterleaved;
 		
-		for( int32 Index = 0; Index < MAX_MULTICHANNEL_AUDIOCHANNELS; Index++ )
+		for( int32 Index = 0; Index < CORE_AUDIO_MAX_MULTICHANNEL_AUDIOCHANNELS; Index++ )
 		{
 			Status = AudioUnitSetProperty( MatrixMixerUnit, kAudioUnitProperty_StreamFormat, kAudioUnitScope_Input, Index, &MatrixMixerInputFormat, Size );
 
@@ -273,12 +273,12 @@ void FCoreAudioDevice::TeardownHardware()
 		MatrixMixerUnit = NULL;
 	}
 	
-	for( int32 Index = 0; Index < MAX_AUDIOCHANNELS; ++Index )
+	for( int32 Index = 0; Index < CORE_AUDIO_MAX_CHANNELS; ++Index )
 	{
 		Mixer3DInputStatus[ Index ] = false;
 	}
 	
-	for( int32 Index = 0; Index < MAX_MULTICHANNEL_AUDIOCHANNELS; ++Index )
+	for( int32 Index = 0; Index < CORE_AUDIO_MAX_MULTICHANNEL_AUDIOCHANNELS; ++Index )
 	{
 		MatrixMixerInputStatus[ Index ] = false;
 	}
@@ -307,7 +307,7 @@ FSoundSource* FCoreAudioDevice::CreateSoundSource()
 
 void FCoreAudioDevice::SetupMatrixMixerInput( int32 Input, bool bIs6ChannelOGG )
 {
-	check( Input < MAX_MULTICHANNEL_AUDIOCHANNELS );
+	check( Input < CORE_AUDIO_MAX_MULTICHANNEL_AUDIOCHANNELS );
 	
 	uint32 InputOffset = Input * MatrixMixerInputFormat.mChannelsPerFrame;
 
@@ -368,7 +368,7 @@ void FCoreAudioDevice::SetupMatrixMixerInput( int32 Input, bool bIs6ChannelOGG )
 
 void FCoreAudioDevice::SetMatrixMixerInputVolume( int32 Input, float Volume )
 {
-	check( Input < MAX_MULTICHANNEL_AUDIOCHANNELS );
+	check( Input < CORE_AUDIO_MAX_MULTICHANNEL_AUDIOCHANNELS );
 	
 	uint32 InputOffset = Input * MatrixMixerInputFormat.mChannelsPerFrame;
 	
@@ -399,7 +399,7 @@ void FCoreAudioDevice::SetMatrixMixerInputVolume( int32 Input, float Volume )
 
 int32 FCoreAudioDevice::GetFreeMixer3DInput()
 {
-	for( int32 Index = 0; Index < MAX_AUDIOCHANNELS; ++Index )
+	for( int32 Index = 0; Index < CORE_AUDIO_MAX_CHANNELS; ++Index )
 	{
 		if( Mixer3DInputStatus[ Index ] == false )
 		{
@@ -418,7 +418,7 @@ void FCoreAudioDevice::SetFreeMixer3DInput( int32 Input )
 
 int32 FCoreAudioDevice::GetFreeMatrixMixerInput()
 {
-	for( int32 Index = 0; Index < MAX_MULTICHANNEL_AUDIOCHANNELS; ++Index )
+	for( int32 Index = 0; Index < CORE_AUDIO_MAX_MULTICHANNEL_AUDIOCHANNELS; ++Index )
 	{
 		if( MatrixMixerInputStatus[ Index ] == false )
 		{
diff --git a/Engine/Source/Runtime/Mac/CoreAudio/Private/CoreAudioEffects.cpp b/Engine/Source/Runtime/Mac/CoreAudio/Private/CoreAudioEffects.cpp
--- a/Engine/Source/Runtime/Mac/CoreAudio/Private/CoreAudioEffects.cpp
+++ b/Engine/Source/Runtime/Mac/CoreAudio/Private/CoreAudioEffects.cpp
@@ -13,7 +13,7 @@
 #include "CoreAudioEffects.h"
 #include "Engine.h"
 
-extern FCoreAudioSoundSource *GAudioChannels[MAX_AUDIOCHANNELS + 1];
+extern FCoreAudioSoundSource *GAudioChannels[CORE_AUDIO_MAX_CHANNELS + 1];
 
 static CFBundleRef LoadRadioEffectComponent()
 {
@@ -132,7 +132,7 @@ void FCoreAudioEffectsManager::SetReverbEffectParameters( const FAudioReverbEffe
 	float LargeDelayRange = 0.3f;														// 0->1, 0.3
 	float LargeBrightness = FMath::Max<float>(ReverbEffectParameters.Density*ReverbEffectParameters.Gain, 0.1f);	// 0.1->1, 0.49
 
-	for( uint32 Index = 1; Index < MAX_AUDIOCHANNELS + 1; Index++ )
+	for( uint32 Index = 1; Index < CORE_AUDIO_MAX_CHANNELS + 1; Index++ )
 	{
 		FCoreAudioSoundSource *Source = GAudioChannels[Index];
 		if( Source && Source->ReverbUnit )
@@ -168,7 +168,7 @@ void FCoreAudioEffectsManager::SetEQEffectParameters( const FAudioEQEffect& EQEf
 	float CenterGain = VolumeToDeciBels(EQEffectParameters.MFGain);
 	float HighGain = VolumeToDeciBels(EQEffectParameters.HFGain);
 
-	for( uint32 Index = 1; Index < MAX_AUDIOCHANNELS + 1; Index++ )
+	for( uint32 Index = 1; Index < CORE_AUDIO_MAX_CHANNELS + 1; Index++ )
 	{
 		FCoreAudioSoundSource *Source = GAudioChannels[Index];
 		if( Source && Source->EQUnit )
@@ -214,7 +214,7 @@ void FCoreAudioEffectsManager::SetRadioEffectParameters( const FAudioRadioEffect
 	const float ChebyshevMultiplier = Radio_ChebyshevMultiplier->GetValueOnGameThread();
 	const float ChebyshevCubedMultiplier = Radio_ChebyshevCubedMultiplier->GetValueOnGameThread();
 
-	for( uint32 Index = 1; Index < MAX_AUDIOCHANNELS + 1; Index++ )
+	for( uint32 Index = 1; Index < CORE_AUDIO_MAX_CHANNELS + 1; Index++ )
 	{
 		FCoreAudioSoundSource *Source = GAudioChannels[Index];
 		if( Source && Source->RadioUnit )
diff --git a/Engine/Source/Runtime/Mac/CoreAudio/Private/CoreAudioSource.cpp b/Engine/Source/Runtime/Mac/CoreAudio/Private/CoreAudioSource.cpp
--- a/Engine/Source/Runtime/Mac/CoreAudio/Private/CoreAudioSource.cpp
+++ b/Engine/Source/Runtime/Mac/CoreAudio/Private/CoreAudioSource.cpp
@@ -25,11 +25,11 @@
 	check(Status == noErr);\
 }
 
-FCoreAudioSoundSource *GAudioChannels[MAX_AUDIOCHANNELS + 1];
+FCoreAudioSoundSource *GAudioChannels[CORE_AUDIO_MAX_CHANNELS + 1];
 
 static int32 FindFreeAudioChannel()
 {
-	for( int32 Index = 1; Index < MAX_AUDIOCHANNELS + 1; Index++ )
+	for( int32 Index = 1; Index < CORE_AUDIO_MAX_CHANNELS + 1; Index++ )
 	{
 		if( GAudioChannels[Index] == NULL )
 		{
@@ -246,7 +246,28 @@ bool FCoreAudioSoundSource::Init( FWaveInstance* InWaveInstance )
 		if( Buffer && Buffer->NumChannels > 0 )
 		{
 			SCOPE_CYCLE_COUNTER( STAT_AudioSourceInitTime );
-		
+
+			if (Buffer->NumChannels < 3)
+			{
+				MixerInputNumber = AudioDevice->GetFreeMixer3DInput();
+			}
+			else
+			{
+				MixerInputNumber = AudioDevice->GetFreeMatrixMixerInput();
+			}
+
+			if (MixerInputNumber == -1)
+			{
+				return false;
+			}
+
+			AudioChannel = FindFreeAudioChannel();
+			if (AudioChannel == 0)
+			{
+				return false;
+			}
+
+
 			WaveInstance = InWaveInstance;
 
 			// Set whether to apply reverb
@@ -292,6 +313,9 @@ void FCoreAudioSoundSource::Update( void )
 		return;
 	}
 
+	check(AudioChannel != 0);
+	check(MixerInputNumber != -1);
+
 	float Volume = 0.0f;
 	
 	if (!AudioDevice->bIsDeviceMuted)
@@ -689,32 +713,45 @@ OSStatus FCoreAudioSoundSource::CreateAndConnectAudioUnit( OSType Type, OSType S
 
 bool FCoreAudioSoundSource::AttachToAUGraph()
 {
-	AudioChannel = FindFreeAudioChannel();
+
+	// We should usually have a non-zero AudioChannel here, but this can happen when unpausing a sound
+	if (AudioChannel == 0)
+	{
+		AudioChannel = FindFreeAudioChannel();
+		if (AudioChannel == 0)
+		{
+			return false;
+		}
+	}
+
+	check(MixerInputNumber != -1);
 
 	OSStatus ErrorStatus = noErr;
 	AUNode DestNode = -1;
-	int32 DestInputNumber = 0;
+	int32 DestInputNumber = MixerInputNumber;
 	AudioStreamBasicDescription* StreamFormat = NULL;
 	
 	if( Buffer->NumChannels < 3 )
 	{
 		ErrorStatus = AudioConverterNew( &Buffer->PCMFormat, &AudioDevice->Mixer3DFormat, &CoreAudioConverter );
 		DestNode = AudioDevice->GetMixer3DNode();
-		MixerInputNumber = DestInputNumber = AudioDevice->GetFreeMixer3DInput();
-		
+
 		uint32 SpatialSetting = ( Buffer->NumChannels == 1 ) ? kSpatializationAlgorithm_SoundField : kSpatializationAlgorithm_StereoPassThrough;
-		AudioUnitSetProperty( AudioDevice->GetMixer3DUnit(), kAudioUnitProperty_SpatializationAlgorithm, kAudioUnitScope_Input, MixerInputNumber, &SpatialSetting, sizeof( SpatialSetting ) );
-		AudioUnitSetParameter( AudioDevice->GetMixer3DUnit(), k3DMixerParam_Distance, kAudioUnitScope_Input, MixerInputNumber, 1.0f, 0 );
+		ErrorStatus = AudioUnitSetProperty( AudioDevice->GetMixer3DUnit(), kAudioUnitProperty_SpatializationAlgorithm, kAudioUnitScope_Input, MixerInputNumber, &SpatialSetting, sizeof( SpatialSetting ) );
+		check(ErrorStatus == noErr);
+
+		ErrorStatus = AudioUnitSetParameter( AudioDevice->GetMixer3DUnit(), k3DMixerParam_Distance, kAudioUnitScope_Input, MixerInputNumber, 1.0f, 0 );
+		check(ErrorStatus == noErr);
 
 		StreamFormat = &AudioDevice->Mixer3DFormat;
 	}
 	else
-	{
-		MixerInputNumber = DestInputNumber = AudioDevice->GetFreeMatrixMixerInput();
+	{	
 		DestNode = AudioDevice->GetMatrixMixerNode();
 		StreamFormat = &AudioDevice->MatrixMixerInputFormat;
 		
 		ErrorStatus = AudioConverterNew( &Buffer->PCMFormat, &AudioDevice->MatrixMixerInputFormat, &CoreAudioConverter );
+		check(ErrorStatus == noErr);
 
 		bool bIs6ChannelOGG = Buffer->NumChannels == 6
 							&& ((Buffer->DecompressionState && Buffer->DecompressionState->UsesVorbisChannelOrdering())
@@ -736,6 +773,7 @@ bool FCoreAudioSoundSource::AttachToAUGraph()
 		++FiltersNeeded;
 	}
 #else
+	check(EQNode == 0);
 	bool bNeedEQFilter = false;
 #endif
 
@@ -746,6 +784,7 @@ bool FCoreAudioSoundSource::AttachToAUGraph()
 		++FiltersNeeded;
 	}
 #else
+	check(RadioNode == 0);
 	bool bNeedRadioFilter = false;
 #endif
 
@@ -756,6 +795,7 @@ bool FCoreAudioSoundSource::AttachToAUGraph()
 		++FiltersNeeded;
 	}
 #else
+	check(ReverbNode == 0);
 	bool bNeedReverbFilter = false;
 #endif
 
@@ -924,7 +964,11 @@ bool FCoreAudioSoundSource::AttachToAUGraph()
 		Input.inputProcRefCon = this;
 		SAFE_CA_CALL( AudioUnitSetProperty( SourceUnit, kAudioUnitProperty_SetRenderCallback, kAudioUnitScope_Input, 0, &Input, sizeof( Input ) ) );
 
-		SAFE_CA_CALL( AUGraphUpdate( AudioDevice->GetAudioUnitGraph(), NULL ) );
+		// This is split to easily investigate the callstack if the assert happens
+		AUGraph Graph = AudioDevice->GetAudioUnitGraph();
+		check(Graph);
+		ErrorStatus = AUGraphUpdate( Graph, NULL );
+        check(ErrorStatus == noErr);
 
 		GAudioChannels[AudioChannel] = this;
 	}
@@ -933,6 +977,9 @@ bool FCoreAudioSoundSource::AttachToAUGraph()
 
 bool FCoreAudioSoundSource::DetachFromAUGraph()
 {
+	check(AudioChannel != 0);
+	check(MixerInputNumber != -1);
+
 	AURenderCallbackStruct Input;
 	Input.inputProc = NULL;
 	Input.inputProcRefCon = NULL;
@@ -942,6 +989,21 @@ bool FCoreAudioSoundSource::DetachFromAUGraph()
 	{
 		SAFE_CA_CALL( AUGraphDisconnectNodeInput( AudioDevice->GetAudioUnitGraph(), StreamSplitterNode, 0 ) );
 	}
+
+// Make sure we still have null nodes
+#if !RADIO_ENABLED
+	check(RadioNode == 0);
+#endif
+
+#if !REVERB_ENABLED
+	check(ReverbNode == 0);
+#endif
+
+#if !EQ_ENABLED
+	check(EQNode == 0);
+#endif
+
+
 	if( ReverbNode )
 	{
 		SAFE_CA_CALL( AUGraphDisconnectNodeInput( AudioDevice->GetAudioUnitGraph(), ReverbNode, 0 ) );
diff --git a/Engine/Source/Runtime/Mac/CoreAudio/Public/CoreAudioDevice.h b/Engine/Source/Runtime/Mac/CoreAudio/Public/CoreAudioDevice.h
--- a/Engine/Source/Runtime/Mac/CoreAudio/Public/CoreAudioDevice.h
+++ b/Engine/Source/Runtime/Mac/CoreAudio/Public/CoreAudioDevice.h
@@ -28,7 +28,16 @@ DECLARE_LOG_CATEGORY_EXTERN(LogCoreAudio, Log, All);
 /**
  * Maximum number of multichannel audio channels - used only for MatrixMixer setup
  */
-#define MAX_MULTICHANNEL_AUDIOCHANNELS 16
+
+//#define CORE_AUDIO_MAX_CHANNELS (MAX_AUDIOCHANNELS)
+//#define CORE_AUDIO_MAX_MULTICHANNEL_AUDIOCHANNELS 16
+
+// Note: I'm reducing the core-audio max channels to 32 since PC uses a max of 32
+// TODO: This should be something we define from an engine .ini.
+#define CORE_AUDIO_MAX_CHANNELS (32)
+
+// Note: reducing the max multichannel files to 8 (from 16) for perf reasons
+#define CORE_AUDIO_MAX_MULTICHANNEL_AUDIOCHANNELS 8
 
 class FCoreAudioDevice;
 class FCoreAudioEffectsManager;
@@ -400,8 +409,8 @@ private:
 	AudioStreamBasicDescription	MatrixMixerInputFormat;
 	AudioStreamBasicDescription	MatrixMixerOutputFormat;
 
-	bool						Mixer3DInputStatus[MAX_MULTICHANNEL_AUDIOCHANNELS];
-	bool						MatrixMixerInputStatus[MAX_AUDIOCHANNELS];
+	bool						Mixer3DInputStatus[CORE_AUDIO_MAX_MULTICHANNEL_AUDIOCHANNELS];
+	bool						MatrixMixerInputStatus[CORE_AUDIO_MAX_CHANNELS];
 
 	friend class FCoreAudioSoundBuffer;
 	friend class FCoreAudioSoundSource;
