UE-18905 ShooterGame crash when closing in game menu

- Re-adding the SAFE_CA_CALL wrapper around AUGraphUpdate since there a legit OSStatus code which can result from API call that isn't noErr. In the case of unpausing sounds, etc, this can happen.

diff --git a/Engine/Source/Runtime/Mac/CoreAudio/Private/CoreAudioSource.cpp b/Engine/Source/Runtime/Mac/CoreAudio/Private/CoreAudioSource.cpp
--- a/Engine/Source/Runtime/Mac/CoreAudio/Private/CoreAudioSource.cpp
+++ b/Engine/Source/Runtime/Mac/CoreAudio/Private/CoreAudioSource.cpp
@@ -967,8 +967,7 @@ bool FCoreAudioSoundSource::AttachToAUGraph()
 		// This is split to easily investigate the callstack if the assert happens
 		AUGraph Graph = AudioDevice->GetAudioUnitGraph();
 		check(Graph);
-		ErrorStatus = AUGraphUpdate( Graph, NULL );
-        check(ErrorStatus == noErr);
+		SAFE_CA_CALL(AUGraphUpdate( Graph, NULL ));
 
 		GAudioChannels[AudioChannel] = this;
 	}
