UE-17938 Fixing some issues that cause mac to sound different than PC client

- Switching to a EQ filter type that is more similar to XAudio2's filter
- Piping the per-voice high-frequency attenuation parameter to a new per-voice low-pass filter
- making the effect processing serial rather than parallel
- Clamping max volumes to 1.0 (linear) before setting volumes on API calls. Not doing this causes sounds to distort as max volume being used was 4.0 and we have volume scale factors greater than 1.0

diff --git a/Engine/Source/Runtime/Mac/CoreAudio/Private/CoreAudioDevice.cpp b/Engine/Source/Runtime/Mac/CoreAudio/Private/CoreAudioDevice.cpp
--- a/Engine/Source/Runtime/Mac/CoreAudio/Private/CoreAudioDevice.cpp
+++ b/Engine/Source/Runtime/Mac/CoreAudio/Private/CoreAudioDevice.cpp
@@ -50,6 +50,11 @@ bool FCoreAudioDevice::InitializeHardware()
 
 	InverseTransform = FMatrix::Identity;
 
+	for (int32 Index = 0; Index < CORE_AUDIO_MAX_CHANNELS + 1; ++Index)
+	{
+		AudioChannels[Index] = nullptr;
+	}
+
 	for( SInt32 Index = 0; Index < CORE_AUDIO_MAX_CHANNELS; ++Index )
 	{
 		Mixer3DInputStatus[ Index ] = false;
@@ -224,6 +229,9 @@ bool FCoreAudioDevice::InitializeHardware()
 		Status = AUGraphConnectNodeInput( AudioUnitGraph, Mixer3DNode, 0, OutputNode, 0 );
 	}
 
+	// Set the sample rate
+	SampleRate = Mixer3DFormat.mSampleRate;
+
 	if( Status != noErr )
 	{
 		UE_LOG(LogInit, Log, TEXT( "Failed to start audio graph!" ) );
diff --git a/Engine/Source/Runtime/Mac/CoreAudio/Private/CoreAudioEffects.cpp b/Engine/Source/Runtime/Mac/CoreAudio/Private/CoreAudioEffects.cpp
--- a/Engine/Source/Runtime/Mac/CoreAudio/Private/CoreAudioEffects.cpp
+++ b/Engine/Source/Runtime/Mac/CoreAudio/Private/CoreAudioEffects.cpp
@@ -13,8 +13,6 @@
 #include "CoreAudioEffects.h"
 #include "Engine.h"
 
-extern FCoreAudioSoundSource *GAudioChannels[CORE_AUDIO_MAX_CHANNELS + 1];
-
 static CFBundleRef LoadRadioEffectComponent()
 {
 	bool bLoaded = false;
@@ -89,7 +87,7 @@ TConsoleVariableData<float>* FCoreAudioEffectsManager::Radio_ChebyshevMultiplier
 FCoreAudioEffectsManager::FCoreAudioEffectsManager( FAudioDevice* InDevice )
 	: FAudioEffectsManager( InDevice )
 {
-#if RADIO_ENABLED
+#if CORE_AUDIO_RADIO_ENABLED
 	RadioBundle = LoadRadioEffectComponent();
 	bRadioAvailable = (RadioBundle != NULL);
 #endif
@@ -97,7 +95,7 @@ FCoreAudioEffectsManager::FCoreAudioEffectsManager( FAudioDevice* InDevice )
 
 FCoreAudioEffectsManager::~FCoreAudioEffectsManager()
 {
-#if RADIO_ENABLED
+#if CORE_AUDIO_RADIO_ENABLED
 	if(RadioBundle)
 	{
 		CFRelease(RadioBundle);
@@ -111,6 +109,7 @@ FCoreAudioEffectsManager::~FCoreAudioEffectsManager()
  */
 void FCoreAudioEffectsManager::SetReverbEffectParameters( const FAudioReverbEffect& ReverbEffectParameters )
 {
+#if CORE_AUDIO_REVERB_ENABLED
 	float DryWetMix = FMath::Sin(ReverbEffectParameters.Volume*M_PI_2) * 100.0f;										// 0.0-100.0, 100.0
 	float SmallLargeMix = ReverbEffectParameters.GainHF * 100.0f;	// 0.0-100.0, 50.0
 	float PreDelay = ReverbEffectParameters.ReflectionsDelay;										// 0.001->0.03, 0.025
@@ -134,7 +133,7 @@ void FCoreAudioEffectsManager::SetReverbEffectParameters( const FAudioReverbEffe
 
 	for( uint32 Index = 1; Index < CORE_AUDIO_MAX_CHANNELS + 1; Index++ )
 	{
-		FCoreAudioSoundSource *Source = GAudioChannels[Index];
+		FCoreAudioSoundSource *Source = ((FCoreAudioDevice*)AudioDevice)->AudioChannels[Index];
 		if( Source && Source->ReverbUnit )
 		{
 			AudioUnitSetParameter(Source->ReverbUnit, kReverbParam_DryWetMix, kAudioUnitScope_Global, 0, DryWetMix, 0);
@@ -157,39 +156,51 @@ void FCoreAudioEffectsManager::SetReverbEffectParameters( const FAudioReverbEffe
 			AudioUnitSetParameter(Source->ReverbUnit, kReverbParam_LargeBrightness, kAudioUnitScope_Global, 0, LargeBrightness, 0);
 		}
 	}
+#endif
 }
 
 /** 
  * Calls the platform specific code to set the parameters that define EQ
  */
-void FCoreAudioEffectsManager::SetEQEffectParameters( const FAudioEQEffect& EQEffectParameters )
+void FCoreAudioEffectsManager::SetEQEffectParameters( const FAudioEQEffect& Params )
 {
-	float LowGain = VolumeToDeciBels(EQEffectParameters.LFGain);
-	float CenterGain = VolumeToDeciBels(EQEffectParameters.MFGain);
-	float HighGain = VolumeToDeciBels(EQEffectParameters.HFGain);
+	float LowGain = VolumeToDeciBels(Params.LFGain);
+	float CenterGain = VolumeToDeciBels(Params.MFGain);
+	float HighGain = VolumeToDeciBels(Params.HFGain);
 
 	for( uint32 Index = 1; Index < CORE_AUDIO_MAX_CHANNELS + 1; Index++ )
 	{
-		FCoreAudioSoundSource *Source = GAudioChannels[Index];
-		if( Source && Source->EQUnit )
+		FCoreAudioSoundSource *Source = ((FCoreAudioDevice*)AudioDevice)->AudioChannels[Index];
+		if (Source)
 		{
-			AudioUnitSetParameter( Source->EQUnit, kMultibandFilter_LowFrequency, kAudioUnitScope_Global, 0, EQEffectParameters.LFFrequency, 0 );
-			AudioUnitSetParameter( Source->EQUnit, kMultibandFilter_LowGain, kAudioUnitScope_Global, 0, LowGain, 0 );
+			if (Source->EQUnit)
+			{
+				AudioUnitSetParameter(Source->EQUnit, kAUNBandEQParam_Frequency + 0,	kAudioUnitScope_Global, 0, Params.LFFrequency,			0);
+				AudioUnitSetParameter(Source->EQUnit, kAUNBandEQParam_Gain + 0,			kAudioUnitScope_Global, 0, LowGain,						0);
+				AudioUnitSetParameter(Source->EQUnit, kAUNBandEQParam_Bandwidth + 0,	kAudioUnitScope_Global, 0, 1.0f,						0);	// from FXEQ_DEFAULT_BANDWIDTH
 
-			AudioUnitSetParameter( Source->EQUnit, kMultibandFilter_CenterFreq1, kAudioUnitScope_Global, 0, (EQEffectParameters.MFCutoffFrequency - EQEffectParameters.LFFrequency) / 2.0f, 0 );
-			AudioUnitSetParameter( Source->EQUnit, kMultibandFilter_CenterGain1, kAudioUnitScope_Global, 0, CenterGain, 0 );
-			AudioUnitSetParameter( Source->EQUnit, kMultibandFilter_Bandwidth1, kAudioUnitScope_Global, 0, EQEffectParameters.MFBandwidth, 0 );
-			
-			AudioUnitSetParameter( Source->EQUnit, kMultibandFilter_CenterFreq2, kAudioUnitScope_Global, 0, EQEffectParameters.MFCutoffFrequency, 0 );
-			AudioUnitSetParameter( Source->EQUnit, kMultibandFilter_CenterGain2, kAudioUnitScope_Global, 0, CenterGain, 0 );
-			AudioUnitSetParameter( Source->EQUnit, kMultibandFilter_Bandwidth2, kAudioUnitScope_Global, 0, EQEffectParameters.MFBandwidth, 0 );
+				AudioUnitSetParameter(Source->EQUnit, kAUNBandEQParam_Frequency + 1,	kAudioUnitScope_Global, 0, Params.MFCutoffFrequency,	0);
+				AudioUnitSetParameter(Source->EQUnit, kAUNBandEQParam_Gain + 1,			kAudioUnitScope_Global, 0, CenterGain,					0);
+				AudioUnitSetParameter(Source->EQUnit, kAUNBandEQParam_Bandwidth + 1,	kAudioUnitScope_Global, 0, Params.MFBandwidth,			0);
 
-			AudioUnitSetParameter( Source->EQUnit, kMultibandFilter_CenterFreq3, kAudioUnitScope_Global, 0, (EQEffectParameters.HFFrequency - EQEffectParameters.MFCutoffFrequency) / 2.0f, 0 );
-			AudioUnitSetParameter( Source->EQUnit, kMultibandFilter_CenterGain3, kAudioUnitScope_Global, 0, CenterGain, 0 );
-			AudioUnitSetParameter( Source->EQUnit, kMultibandFilter_Bandwidth3, kAudioUnitScope_Global, 0, EQEffectParameters.MFBandwidth, 0 );
+				AudioUnitSetParameter(Source->EQUnit, kAUNBandEQParam_Frequency + 2,	kAudioUnitScope_Global, 0, Params.HFFrequency,			0);
+				AudioUnitSetParameter(Source->EQUnit, kAUNBandEQParam_Gain + 2,			kAudioUnitScope_Global, 0, HighGain,					0);
+				AudioUnitSetParameter(Source->EQUnit, kAUNBandEQParam_Bandwidth + 2,	kAudioUnitScope_Global, 0, 1.0f,						0); // from FXEQ_DEFAULT_BANDWIDTH
 
-			AudioUnitSetParameter( Source->EQUnit, kMultibandFilter_HighFrequency, kAudioUnitScope_Global, 0, EQEffectParameters.HFFrequency, 0 );
-			AudioUnitSetParameter( Source->EQUnit, kMultibandFilter_HighGain, kAudioUnitScope_Global, 0, HighGain, 0 );
+				AudioUnitSetParameter(Source->EQUnit, kAUNBandEQParam_Frequency + 3,	kAudioUnitScope_Global, 0, 10000.0f,					0);	// from FXEQ_DEFAULT_CENTER_3
+				AudioUnitSetParameter(Source->EQUnit, kAUNBandEQParam_Gain + 3,			kAudioUnitScope_Global, 0, 1.0,							0); // from FXEQ_DEFAULT_GAIN
+				AudioUnitSetParameter(Source->EQUnit, kAUNBandEQParam_Bandwidth + 3,	kAudioUnitScope_Global, 0, 1.0f,						0); // from FXEQ_DEFAULT_BANDWIDTH
+			}
+
+			if (Source->LowPassUnit && Source->HighFrequencyGain < 1.0f - KINDA_SMALL_NUMBER)
+			{
+				float RadianFrequency = 2.0f * FMath::Sin( PI * 6000.0f * Source->HighFrequencyGain / 48000.0f );
+				float CuttoffFrequency = RadianFrequency * ((FCoreAudioDevice*)AudioDevice)->SampleRate;
+				float OneOverQ = ((FCoreAudioDevice*)AudioDevice)->GetLowPassFilterResonance();
+
+				AudioUnitSetParameter(Source->LowPassUnit, kLowPassParam_CutoffFrequency, kAudioUnitScope_Global, 0, CuttoffFrequency, 0);
+				AudioUnitSetParameter(Source->LowPassUnit, kLowPassParam_Resonance, kAudioUnitScope_Global, 0, OneOverQ, 0);
+			}
 		}
 	}
 }
@@ -216,7 +227,7 @@ void FCoreAudioEffectsManager::SetRadioEffectParameters( const FAudioRadioEffect
 
 	for( uint32 Index = 1; Index < CORE_AUDIO_MAX_CHANNELS + 1; Index++ )
 	{
-		FCoreAudioSoundSource *Source = GAudioChannels[Index];
+		FCoreAudioSoundSource *Source = ((FCoreAudioDevice*)AudioDevice)->AudioChannels[Index];
 		if( Source && Source->RadioUnit )
 		{
 			AudioUnitSetParameter( Source->RadioUnit, RadioParam_ChebyshevPowerMultiplier, kAudioUnitScope_Global, 0, ChebyshevPowerMultiplier, 0 );
diff --git a/Engine/Source/Runtime/Mac/CoreAudio/Private/CoreAudioSource.cpp b/Engine/Source/Runtime/Mac/CoreAudio/Private/CoreAudioSource.cpp
--- a/Engine/Source/Runtime/Mac/CoreAudio/Private/CoreAudioSource.cpp
+++ b/Engine/Source/Runtime/Mac/CoreAudio/Private/CoreAudioSource.cpp
@@ -25,20 +25,6 @@
 	check(Status == noErr);\
 }
 
-FCoreAudioSoundSource *GAudioChannels[CORE_AUDIO_MAX_CHANNELS + 1];
-
-static int32 FindFreeAudioChannel()
-{
-	for( int32 Index = 1; Index < CORE_AUDIO_MAX_CHANNELS + 1; Index++ )
-	{
-		if( GAudioChannels[Index] == NULL )
-		{
-			return Index;
-		}
-	}
-
-	return 0;
-}
 
 /*------------------------------------------------------------------------------------
  FCoreAudioSoundSource.
@@ -56,10 +42,10 @@ FCoreAudioSoundSource::FCoreAudioSoundSource( FAudioDevice* InAudioDevice )
 	bBuffersToFlush( false ),
 	SourceNode( 0 ),
 	SourceUnit( NULL ),
-	StreamSplitterNode( 0 ),
-	StreamSplitterUnit( NULL ),
 	EQNode( 0 ),
 	EQUnit( NULL ),
+	LowPassNode( 0 ),
+	LowPassUnit( NULL ),
 	RadioNode( 0 ),
 	RadioUnit( NULL ),
 	bRadioMuted( false ),
@@ -67,8 +53,6 @@ FCoreAudioSoundSource::FCoreAudioSoundSource( FAudioDevice* InAudioDevice )
 	ReverbUnit( NULL ),
 	bReverbMuted( false ),
 	bDryMuted( false ),
-	StreamMergerNode( 0 ),
-	StreamMergerUnit( NULL ),
 	AudioChannel( 0 ),
 	BufferInUse( 0 ),
 	NumActiveBuffers( 0 ),
@@ -261,7 +245,7 @@ bool FCoreAudioSoundSource::Init( FWaveInstance* InWaveInstance )
 				return false;
 			}
 
-			AudioChannel = FindFreeAudioChannel();
+			AudioChannel = AudioDevice->FindFreeAudioChannel();
 			if (AudioChannel == 0)
 			{
 				return false;
@@ -367,95 +351,7 @@ void FCoreAudioSoundSource::Update( void )
 			Elevation = Rotation.Pitch;
 		}
 
-		// Apply any debug settings
-		{
-			bool bMuteDry, bMuteReverb, bMuteRadio;
-			switch( AudioDevice->GetMixDebugState() )
-			{
-			case DEBUGSTATE_IsolateReverb:		bMuteDry = true;	bMuteReverb = false;	bMuteRadio = false;	break;
-			case DEBUGSTATE_IsolateDryAudio:	bMuteDry = false;	bMuteReverb = true;		bMuteRadio = true;	break;
-			default:							bMuteDry = false;	bMuteReverb = false;	bMuteRadio = false;	break;
-			};
-
-			// Dry audio or EQ node
-			if( bMuteDry != bDryMuted )
-			{
-				if( bMuteDry )
-				{
-					if( StreamSplitterUnit )
-					{
-						for( int32 OutputChannelIndex = 0; OutputChannelIndex < AudioDevice->Mixer3DFormat.mChannelsPerFrame; ++OutputChannelIndex )
-						{
-							SAFE_CA_CALL( AudioUnitSetParameter( StreamSplitterUnit, kMatrixMixerParam_Volume, kAudioUnitScope_Output, OutputChannelIndex, 0.0, 0 ) );
-						}
-					}
-					else
-					{
-						SAFE_CA_CALL( AudioUnitSetParameter( AudioDevice->GetMixer3DUnit(), k3DMixerParam_Gain, kAudioUnitScope_Input, MixerInputNumber, 0.0, 0 ) );
-					}
-					bDryMuted = true;
-				}
-				else
-				{
-					if( StreamSplitterUnit )
-					{
-						for( int32 OutputChannelIndex = 0; OutputChannelIndex < AudioDevice->Mixer3DFormat.mChannelsPerFrame; ++OutputChannelIndex )
-						{
-							SAFE_CA_CALL( AudioUnitSetParameter( StreamSplitterUnit, kMatrixMixerParam_Volume, kAudioUnitScope_Output, OutputChannelIndex, 1.0, 0 ) );
-						}
-					}
-					bDryMuted = false;
-				}
-			}
-
-			if( ReverbNode && bMuteReverb != bReverbMuted )
-			{
-				int ReverbNodeBaseIndex = AudioDevice->Mixer3DFormat.mChannelsPerFrame;
-				if( bMuteReverb )
-				{
-					for( int32 OutputChannelIndex = 0; OutputChannelIndex < AudioDevice->Mixer3DFormat.mChannelsPerFrame; ++OutputChannelIndex )
-					{
-						SAFE_CA_CALL( AudioUnitSetParameter( StreamSplitterUnit, kMatrixMixerParam_Volume, kAudioUnitScope_Output, ReverbNodeBaseIndex + OutputChannelIndex, 0.0, 0 ) );
-					}
-					bReverbMuted = true;
-				}
-				else
-				{
-					for( int32 OutputChannelIndex = 0; OutputChannelIndex < AudioDevice->Mixer3DFormat.mChannelsPerFrame; ++OutputChannelIndex )
-					{
-						SAFE_CA_CALL( AudioUnitSetParameter( StreamSplitterUnit, kMatrixMixerParam_Volume, kAudioUnitScope_Output, ReverbNodeBaseIndex + OutputChannelIndex, 1.0, 0 ) );
-					}
-					bReverbMuted = false;
-				}
-			}
-
-			if( RadioNode && bMuteRadio != bRadioMuted )
-			{
-				int RadioNodeBaseIndex = ( 1 + ( ReverbNode ? 1 : 0 ) ) * AudioDevice->Mixer3DFormat.mChannelsPerFrame;
-				if( bMuteRadio )
-				{
-					for( int32 OutputChannelIndex = 0; OutputChannelIndex < AudioDevice->Mixer3DFormat.mChannelsPerFrame; ++OutputChannelIndex )
-					{
-						SAFE_CA_CALL( AudioUnitSetParameter( StreamSplitterUnit, kMatrixMixerParam_Volume, kAudioUnitScope_Output, RadioNodeBaseIndex + OutputChannelIndex, 0.0, 0 ) );
-					}
-					bRadioMuted = true;
-				}
-				else
-				{
-					for( int32 OutputChannelIndex = 0; OutputChannelIndex < AudioDevice->Mixer3DFormat.mChannelsPerFrame; ++OutputChannelIndex )
-					{
-						SAFE_CA_CALL( AudioUnitSetParameter( StreamSplitterUnit, kMatrixMixerParam_Volume, kAudioUnitScope_Output, RadioNodeBaseIndex + OutputChannelIndex, 1.0, 0 ) );
-					}
-					bRadioMuted = false;
-				}
-			}
-		}
-
-		if( !bDryMuted || StreamSplitterUnit )
-		{
-			SAFE_CA_CALL( AudioUnitSetParameter( AudioDevice->GetMixer3DUnit(), k3DMixerParam_Gain, kAudioUnitScope_Input, MixerInputNumber, Volume, 0 ) );
-		}
-
+		SAFE_CA_CALL( AudioUnitSetParameter( AudioDevice->GetMixer3DUnit(), k3DMixerParam_Gain, kAudioUnitScope_Input, MixerInputNumber, Volume, 0 ) );
 		SAFE_CA_CALL( AudioUnitSetParameter( AudioDevice->GetMixer3DUnit(), k3DMixerParam_PlaybackRate, kAudioUnitScope_Input, MixerInputNumber, Pitch, 0 ) );
 		SAFE_CA_CALL( AudioUnitSetParameter( AudioDevice->GetMixer3DUnit(), k3DMixerParam_Azimuth, kAudioUnitScope_Input, MixerInputNumber, Azimuth, 0 ) );
 		SAFE_CA_CALL( AudioUnitSetParameter( AudioDevice->GetMixer3DUnit(), k3DMixerParam_Elevation, kAudioUnitScope_Input, MixerInputNumber, Elevation, 0 ) );
@@ -711,13 +607,162 @@ OSStatus FCoreAudioSoundSource::CreateAndConnectAudioUnit( OSType Type, OSType S
 	return Status;
 }
 
-bool FCoreAudioSoundSource::AttachToAUGraph()
+void FCoreAudioSoundSource::InitSourceUnit(AudioStreamBasicDescription* StreamFormat, AUNode& HeadNode)
+{
+	AudioComponentDescription Desc = {kAudioUnitType_FormatConverter, kAudioUnitSubType_AUConverter, kAudioUnitManufacturer_Apple, 0, 0};
+	OSStatus ErrorStatus = AUGraphAddNode(AudioDevice->GetAudioUnitGraph(), &Desc, &SourceNode);
+	check(ErrorStatus == noErr);
+
+	ErrorStatus = AUGraphNodeInfo(AudioDevice->GetAudioUnitGraph(), SourceNode, nullptr, &SourceUnit);
+	check(ErrorStatus == noErr);
+
+	ErrorStatus = AudioUnitSetProperty(SourceUnit, kAudioUnitProperty_StreamFormat, kAudioUnitScope_Input, 0, StreamFormat, sizeof(AudioStreamBasicDescription));
+	check(ErrorStatus == noErr);
+
+	ErrorStatus = AudioUnitSetProperty(SourceUnit, kAudioUnitProperty_StreamFormat, kAudioUnitScope_Output, 0, StreamFormat, sizeof(AudioStreamBasicDescription));
+	check(ErrorStatus == noErr);
+
+	// Setup the callback which feeds audio to the source audio unit
+	AURenderCallbackStruct Input;
+	Input.inputProc = &CoreAudioRenderCallback;
+	Input.inputProcRefCon = this;
+	SAFE_CA_CALL( AudioUnitSetProperty(SourceUnit, kAudioUnitProperty_SetRenderCallback, kAudioUnitScope_Input, 0, &Input, sizeof(Input)));
+
+	HeadNode = SourceNode;
+}
+
+void FCoreAudioSoundSource::InitLowPassEffect(AudioStreamBasicDescription* StreamFormat, AUNode& HeadNode)
+{
+	AudioComponentDescription Desc = {kAudioUnitType_Effect, kAudioUnitSubType_LowPassFilter, kAudioUnitManufacturer_Apple, 0, 0};
+	OSStatus ErrorStatus = AUGraphAddNode(AudioDevice->GetAudioUnitGraph(), &Desc, &LowPassNode);
+	check(ErrorStatus == noErr);
+
+	ErrorStatus = AUGraphNodeInfo(AudioDevice->GetAudioUnitGraph(), LowPassNode, nullptr, &LowPassUnit);
+	check(ErrorStatus == noErr);
+
+	ErrorStatus = AudioUnitSetProperty(LowPassUnit, kAudioUnitProperty_StreamFormat, kAudioUnitScope_Input, 0, StreamFormat, sizeof(AudioStreamBasicDescription));
+	check(ErrorStatus == noErr);
+
+	ErrorStatus = AudioUnitSetProperty(LowPassUnit, kAudioUnitProperty_StreamFormat, kAudioUnitScope_Output, 0, StreamFormat, sizeof(AudioStreamBasicDescription));
+	check(ErrorStatus == noErr);
+
+	// Set the cutoff frequency to be the nyquist at first
+	float CutoffFreq = AudioDevice->SampleRate * 0.5f;
+	ErrorStatus = AudioUnitSetParameter(LowPassUnit, kLowPassParam_CutoffFrequency,	kAudioUnitScope_Global, 0, CutoffFreq, 0);
+	check(ErrorStatus == noErr);
+
+	ErrorStatus = AudioUnitInitialize(LowPassUnit);
+	check(ErrorStatus == noErr);
+
+	// Connect the current head node to the radio node (i.e. source -> radio effect)
+	ErrorStatus = AUGraphConnectNodeInput(AudioDevice->GetAudioUnitGraph(), HeadNode, 0, LowPassNode, 0);
+	check(ErrorStatus == noErr);
+
+	// The radio node becomes the head node
+	HeadNode = LowPassNode;
+}
+
+void FCoreAudioSoundSource::InitRadioSourceEffect(AudioStreamBasicDescription* StreamFormat, AUNode& HeadNode)
+{
+	AudioComponentDescription Desc = {kAudioUnitType_Effect, 'Rdio', 'Epic', 0, 0};
+	OSStatus ErrorStatus = AUGraphAddNode(AudioDevice->GetAudioUnitGraph(), &Desc, &RadioNode);
+	check(ErrorStatus == noErr);
+
+	ErrorStatus = AUGraphNodeInfo(AudioDevice->GetAudioUnitGraph(), RadioNode, nullptr, &RadioUnit);
+	check(ErrorStatus == noErr);
+
+	ErrorStatus = AudioUnitSetProperty(RadioUnit, kAudioUnitProperty_StreamFormat, kAudioUnitScope_Input, 0, StreamFormat, sizeof(AudioStreamBasicDescription));
+	check(ErrorStatus == noErr);
+
+	ErrorStatus = AudioUnitSetProperty(RadioUnit, kAudioUnitProperty_StreamFormat, kAudioUnitScope_Output, 0, StreamFormat, sizeof(AudioStreamBasicDescription));
+	check(ErrorStatus == noErr);
+
+	ErrorStatus = AudioUnitInitialize(RadioUnit);
+	check(ErrorStatus == noErr);
+
+	// Connect the current head node to the radio node (i.e. source -> radio effect)
+	ErrorStatus = AUGraphConnectNodeInput(AudioDevice->GetAudioUnitGraph(), HeadNode, 0, RadioNode, 0);
+	check(ErrorStatus == noErr);
+
+	// The radio node becomes the head node
+	HeadNode = RadioNode;
+}
+
+void FCoreAudioSoundSource::InitEqSourceEffect(AudioStreamBasicDescription* StreamFormat, AUNode& HeadNode)
 {
+	AudioComponentDescription Desc = { kAudioUnitType_Effect, kAudioUnitSubType_NBandEQ, kAudioUnitManufacturer_Apple, 0, 0 };
 
+	OSStatus ErrorStatus = AUGraphAddNode(AudioDevice->GetAudioUnitGraph(), &Desc, &EQNode);
+	check(ErrorStatus == noErr);
+
+	ErrorStatus = AUGraphNodeInfo(AudioDevice->GetAudioUnitGraph(), EQNode, nullptr, &EQUnit);
+	check(ErrorStatus == noErr);
+
+	UInt32 NumBands = 4;
+	ErrorStatus = AudioUnitSetProperty(EQUnit, kAUNBandEQProperty_NumberOfBands, kAudioUnitScope_Global, 0, &NumBands, sizeof(NumBands));
+	check(ErrorStatus == noErr);
+
+	ErrorStatus = AudioUnitSetProperty(EQUnit, kAudioUnitProperty_StreamFormat, kAudioUnitScope_Input, 0, StreamFormat, sizeof(AudioStreamBasicDescription));
+	check(ErrorStatus == noErr);
+
+	ErrorStatus = AudioUnitSetProperty(EQUnit, kAudioUnitProperty_StreamFormat, kAudioUnitScope_Output, 0, StreamFormat, sizeof(AudioStreamBasicDescription));
+	check(ErrorStatus == noErr);
+
+	for (int32 Band = 0; Band < 4; ++Band)
+	{
+		// Now set the filter types for each band
+		ErrorStatus = AudioUnitSetParameter(EQUnit, kAUNBandEQParam_FilterType + Band, kAudioUnitScope_Global, 0, kAUNBandEQFilterType_Parametric, 0);
+		check(ErrorStatus == noErr);
+
+		// Now make sure the bands are not bypassed
+		ErrorStatus = AudioUnitSetParameter(EQUnit, kAUNBandEQParam_BypassBand + Band, kAudioUnitScope_Global, 0, 0, 0);
+		check(ErrorStatus == noErr);
+	}
+
+	ErrorStatus = AudioUnitInitialize(EQUnit);
+	check(ErrorStatus == noErr);
+
+	// Connect the current head node to the radio node (i.e. head -> eq effect)
+	ErrorStatus = AUGraphConnectNodeInput(AudioDevice->GetAudioUnitGraph(), HeadNode, 0, EQNode, 0);
+	check(ErrorStatus == noErr);
+
+	// The radio node becomes the head node
+	HeadNode = EQNode;
+}
+
+void FCoreAudioSoundSource::InitReverbSourceEffect(AudioStreamBasicDescription* StreamFormat, AUNode& HeadNode)
+{
+	AudioComponentDescription Desc = {kAudioUnitType_Effect, kAudioUnitSubType_MatrixReverb, kAudioUnitManufacturer_Apple, 0, 0};
+	OSStatus ErrorStatus = AUGraphAddNode(AudioDevice->GetAudioUnitGraph(), &Desc, &ReverbNode);
+	check(ErrorStatus == noErr);
+
+	ErrorStatus = AUGraphNodeInfo(AudioDevice->GetAudioUnitGraph(), ReverbNode, nullptr, &ReverbUnit);
+	check(ErrorStatus == noErr);
+
+	ErrorStatus = AudioUnitSetProperty(ReverbUnit, kAudioUnitProperty_StreamFormat, kAudioUnitScope_Input, 0, StreamFormat, sizeof(AudioStreamBasicDescription));
+	check(ErrorStatus == noErr);
+
+	ErrorStatus = AudioUnitSetProperty(ReverbUnit, kAudioUnitProperty_StreamFormat, kAudioUnitScope_Output, 0, StreamFormat, sizeof(AudioStreamBasicDescription));
+	check(ErrorStatus == noErr);
+
+	ErrorStatus = AudioUnitInitialize(ReverbUnit);
+	check(ErrorStatus == noErr);
+
+	// Connect the current head node to the input of the reverb node (i.e. head -> reverb effect)
+	ErrorStatus = AUGraphConnectNodeInput(AudioDevice->GetAudioUnitGraph(), HeadNode, 0, ReverbNode, 0);
+	check(ErrorStatus == noErr);
+
+	// The radio node becomes the head node
+	HeadNode = ReverbNode;
+}
+
+
+bool FCoreAudioSoundSource::AttachToAUGraph()
+{
 	// We should usually have a non-zero AudioChannel here, but this can happen when unpausing a sound
 	if (AudioChannel == 0)
 	{
-		AudioChannel = FindFreeAudioChannel();
+		AudioChannel = AudioDevice->FindFreeAudioChannel();
 		if (AudioChannel == 0)
 		{
 			return false;
@@ -727,14 +772,14 @@ bool FCoreAudioSoundSource::AttachToAUGraph()
 	check(MixerInputNumber != -1);
 
 	OSStatus ErrorStatus = noErr;
-	AUNode DestNode = -1;
-	int32 DestInputNumber = MixerInputNumber;
+	AUNode HeadNode = -1;
+	AUNode FinalNode = -1;
 	AudioStreamBasicDescription* StreamFormat = NULL;
 	
-	if( Buffer->NumChannels < 3 )
+	if (Buffer->NumChannels < 3)
 	{
 		ErrorStatus = AudioConverterNew( &Buffer->PCMFormat, &AudioDevice->Mixer3DFormat, &CoreAudioConverter );
-		DestNode = AudioDevice->GetMixer3DNode();
+		FinalNode = AudioDevice->GetMixer3DNode();
 
 		uint32 SpatialSetting = ( Buffer->NumChannels == 1 ) ? kSpatializationAlgorithm_SoundField : kSpatializationAlgorithm_StereoPassThrough;
 		ErrorStatus = AudioUnitSetProperty( AudioDevice->GetMixer3DUnit(), kAudioUnitProperty_SpatializationAlgorithm, kAudioUnitScope_Input, MixerInputNumber, &SpatialSetting, sizeof( SpatialSetting ) );
@@ -747,7 +792,7 @@ bool FCoreAudioSoundSource::AttachToAUGraph()
 	}
 	else
 	{	
-		DestNode = AudioDevice->GetMatrixMixerNode();
+		FinalNode = AudioDevice->GetMatrixMixerNode();
 		StreamFormat = &AudioDevice->MatrixMixerInputFormat;
 		
 		ErrorStatus = AudioConverterNew( &Buffer->PCMFormat, &AudioDevice->MatrixMixerInputFormat, &CoreAudioConverter );
@@ -757,219 +802,64 @@ bool FCoreAudioSoundSource::AttachToAUGraph()
 							&& ((Buffer->DecompressionState && Buffer->DecompressionState->UsesVorbisChannelOrdering())
 							|| WaveInstance->WaveData->bDecompressedFromOgg);
 
-		AudioDevice->SetupMatrixMixerInput( DestInputNumber, bIs6ChannelOGG );
+		AudioDevice->SetupMatrixMixerInput( MixerInputNumber, bIs6ChannelOGG );
 	}
 
-	if( ErrorStatus != noErr )
-	{
-		UE_LOG(LogCoreAudio, Warning, TEXT("CoreAudioConverter creation failed, error code %d"), ErrorStatus);
-	}
+	// Initiliaze the "source" node, the node that is generating audio
+	// This node becomes the "head" node
+	InitSourceUnit(StreamFormat, HeadNode);
 
-	int FiltersNeeded = 0;
-#if EQ_ENABLED
-	bool bNeedEQFilter = IsEQFilterApplied();
-	if( bNeedEQFilter )
-	{
-		++FiltersNeeded;
-	}
-#else
-	check(EQNode == 0);
-	bool bNeedEQFilter = false;
-#endif
+	// Figure out what filters are needed
 
-#if RADIO_ENABLED
-	bool bNeedRadioFilter = Effects->bRadioAvailable && WaveInstance->bApplyRadioFilter;
-	if( bNeedRadioFilter )
-	{
-		++FiltersNeeded;
-	}
-#else
-	check(RadioNode == 0);
+	bool bNeedEQFilter = false;
 	bool bNeedRadioFilter = false;
-#endif
-
-#if REVERB_ENABLED
-	bool bNeedReverbFilter = bReverbApplied;
-	if( bNeedReverbFilter )
-	{
-		++FiltersNeeded;
-	}
-#else
-	check(ReverbNode == 0);
 	bool bNeedReverbFilter = false;
-#endif
-
-	if( FiltersNeeded > 0 )
-	{
-		uint32 BusCount = FiltersNeeded + ( bNeedEQFilter ? 0 : 1 );	// one for each filter, plus one for dry voice if there's no EQ filter
-
-		// Prepare Voice Merger
-		SAFE_CA_CALL( CreateAudioUnit( kAudioUnitType_Mixer, kAudioUnitSubType_MatrixMixer, kAudioUnitManufacturer_Apple, NULL, NULL, &StreamMergerNode, &StreamMergerUnit ) );
-
-		// Set Bus Counts
-		uint32 NumBuses = BusCount;
-		SAFE_CA_CALL( AudioUnitSetProperty( StreamMergerUnit, kAudioUnitProperty_ElementCount, kAudioUnitScope_Input, 0, &NumBuses, sizeof(uint32) ) );
-		NumBuses = 1;
-		SAFE_CA_CALL( AudioUnitSetProperty(	StreamMergerUnit, kAudioUnitProperty_ElementCount, kAudioUnitScope_Output, 0, &NumBuses, sizeof(uint32) ) );
-
-		// Set Input Formats
-		for( int32 InputIndex = 0; InputIndex < BusCount; ++InputIndex )
-		{
-			SAFE_CA_CALL( AudioUnitSetProperty( StreamMergerUnit, kAudioUnitProperty_StreamFormat, kAudioUnitScope_Input, InputIndex, StreamFormat, sizeof( AudioStreamBasicDescription ) ) );
-		}
-
-		// Set Output Format
-		SAFE_CA_CALL( AudioUnitSetProperty( StreamMergerUnit, kAudioUnitProperty_StreamFormat, kAudioUnitScope_Output, 0, StreamFormat, sizeof( AudioStreamBasicDescription ) ) );
 
-		SAFE_CA_CALL( AudioUnitInitialize( StreamMergerUnit ) );
-
-		// Set Master volume
-		SAFE_CA_CALL( AudioUnitSetParameter( StreamMergerUnit, kMatrixMixerParam_Volume, kAudioUnitScope_Global, 0xFFFFFFFF, 1.0, 0 ) );
-
-		// Enable Output
-		SAFE_CA_CALL( AudioUnitSetParameter( StreamMergerUnit, kMatrixMixerParam_Enable, kAudioUnitScope_Output, 0, 1.0, 0 ) );
-
-		// Set Output volumes
-		for( int32 ChannelIndex = 0; ChannelIndex < StreamFormat->mChannelsPerFrame; ++ChannelIndex )
-		{
-			SAFE_CA_CALL( AudioUnitSetParameter( StreamMergerUnit, kMatrixMixerParam_Volume, kAudioUnitScope_Output, ChannelIndex, 1.0, 0 ) );
-		}
-
-		for( int32 InputIndex = 0; InputIndex < BusCount; ++InputIndex )
-		{
-			// Enable Input
-			SAFE_CA_CALL( AudioUnitSetParameter( StreamMergerUnit, kMatrixMixerParam_Enable, kAudioUnitScope_Input, InputIndex, 1.0, 0 ) );
-		}
-
-		for( int32 ChannelIndex = 0; ChannelIndex < StreamFormat->mChannelsPerFrame; ++ChannelIndex )
-		{
-			for( int32 InputBusIndex = 0; InputBusIndex < BusCount; ++InputBusIndex )
-			{
-				int32 InputChannelIndex = InputBusIndex*StreamFormat->mChannelsPerFrame + ChannelIndex;
-
-				// Set Input Channel Volume
-				SAFE_CA_CALL( AudioUnitSetParameter( StreamMergerUnit, kMatrixMixerParam_Volume, kAudioUnitScope_Input, InputChannelIndex, 1.0, 0 ) );
-
-				// Set Crossfade Volume - each input channel goes to specific output channel. The rest of connections is left at zero.
-				SAFE_CA_CALL( AudioUnitSetParameter( StreamMergerUnit, kMatrixMixerParam_Volume, kAudioUnitScope_Global, ( InputChannelIndex << 16 ) | ChannelIndex, 1.0, 0 ) );
-			}
-		}
-
-		SAFE_CA_CALL( AUGraphConnectNodeInput( AudioDevice->GetAudioUnitGraph(), StreamMergerNode, 0, DestNode, DestInputNumber ) );
-
-		DestNode = StreamMergerNode;
-		DestInputNumber = 0;
-
-		// Prepare and initialize stream splitter
-		SAFE_CA_CALL( CreateAudioUnit( kAudioUnitType_Mixer, kAudioUnitSubType_MatrixMixer, kAudioUnitManufacturer_Apple, NULL, NULL, &StreamSplitterNode, &StreamSplitterUnit ) );
-
-		// Set bus counts
-		NumBuses = 1;
-		SAFE_CA_CALL( AudioUnitSetProperty( StreamSplitterUnit, kAudioUnitProperty_ElementCount, kAudioUnitScope_Input, 0, &NumBuses, sizeof(uint32) ) );
-		NumBuses = BusCount;
-		SAFE_CA_CALL( AudioUnitSetProperty(	StreamSplitterUnit, kAudioUnitProperty_ElementCount, kAudioUnitScope_Output, 0, &NumBuses, sizeof(uint32) ) );
-
-		// Set Input format
-		SAFE_CA_CALL( AudioUnitSetProperty( StreamSplitterUnit, kAudioUnitProperty_StreamFormat, kAudioUnitScope_Input, 0, StreamFormat, sizeof( AudioStreamBasicDescription ) ) );
-
-		// Set Output formats
-		for( int32 OutputIndex = 0; OutputIndex < BusCount; ++OutputIndex )
-		{
-			SAFE_CA_CALL( AudioUnitSetProperty( StreamSplitterUnit, kAudioUnitProperty_StreamFormat, kAudioUnitScope_Output, OutputIndex, StreamFormat, sizeof( AudioStreamBasicDescription ) ) );
-		}
-
-		SAFE_CA_CALL( AudioUnitInitialize( StreamSplitterUnit ) );
-
-		// Set Master volume
-		SAFE_CA_CALL( AudioUnitSetParameter( StreamSplitterUnit, kMatrixMixerParam_Volume, kAudioUnitScope_Global, 0xFFFFFFFF, 1.0, 0 ) );
-
-		// Enable Input
-		SAFE_CA_CALL( AudioUnitSetParameter( StreamSplitterUnit, kMatrixMixerParam_Enable, kAudioUnitScope_Input, 0, 1.0, 0 ) );
-
-		// Set Input Volumes
-		for( int32 ChannelIndex = 0; ChannelIndex < StreamFormat->mChannelsPerFrame; ++ChannelIndex )
-		{
-			SAFE_CA_CALL( AudioUnitSetParameter( StreamSplitterUnit, kMatrixMixerParam_Volume, kAudioUnitScope_Input, ChannelIndex, 1.0, 0 ) );
-		}
-
-		for( int32 OutputIndex = 0; OutputIndex < BusCount; ++OutputIndex )
-		{
-			// Enable Output
-			SAFE_CA_CALL( AudioUnitSetParameter( StreamSplitterUnit, kMatrixMixerParam_Enable, kAudioUnitScope_Output, OutputIndex, 1.0, 0 ) );
-		}
-
-		for( int32 ChannelIndex = 0; ChannelIndex < StreamFormat->mChannelsPerFrame; ++ChannelIndex )
-		{
-			for( int32 OutputBusIndex = 0; OutputBusIndex < BusCount; ++OutputBusIndex )
-			{
-				int32 OutputChannelIndex = OutputBusIndex*StreamFormat->mChannelsPerFrame + ChannelIndex;
-
-				// Set Output Channel Volume
-				SAFE_CA_CALL( AudioUnitSetParameter( StreamSplitterUnit, kMatrixMixerParam_Volume, kAudioUnitScope_Output, OutputChannelIndex, 1.0, 0 ) );
-
-				// Set Crossfade Volume - each output channel goes from specific input channel. The rest of connections is left at zero.
-				SAFE_CA_CALL( AudioUnitSetParameter( StreamSplitterUnit, kMatrixMixerParam_Volume, kAudioUnitScope_Global, ( ChannelIndex << 16 ) | OutputChannelIndex, 1.0, 0 ) );
-			}
-		}
+#if CORE_AUDIO_EQ_ENABLED
+	bNeedEQFilter = IsEQFilterApplied();
+#endif
 
-		// Prepare and connect appropriate filters
-#if EQ_ENABLED
-		if( bNeedEQFilter )
-		{
-			SAFE_CA_CALL( CreateAndConnectAudioUnit( kAudioUnitType_Effect, kAudioUnitSubType_AUFilter, kAudioUnitManufacturer_Apple, DestNode, DestInputNumber, StreamFormat, StreamFormat, &EQNode, &EQUnit ) );
-			SAFE_CA_CALL( AUGraphConnectNodeInput( AudioDevice->GetAudioUnitGraph(), StreamSplitterNode, DestInputNumber, EQNode, 0 ) );
-		}
-		else
+#if CORE_AUDIO_RADIO_ENABLED
+	bNeedRadioFilter = Effects->bRadioAvailable && WaveInstance->bApplyRadioFilter;
 #endif
-		{
-			// Add direct connection between stream splitter and stream merger, for dry voice
-			SAFE_CA_CALL( AUGraphConnectNodeInput( AudioDevice->GetAudioUnitGraph(), StreamSplitterNode, 0, StreamMergerNode, 0 ) );
-
-			// Silencing dry voice (for testing)
-//			for( int32 ChannelIndex = 0; ChannelIndex < StreamFormat->mChannelsPerFrame; ++ChannelIndex )
-//			{
-//				SAFE_CA_CALL( AudioUnitSetParameter( StreamSplitterUnit, kMatrixMixerParam_Volume, kAudioUnitScope_Output, ChannelIndex, 0.0, 0 ) );
-//			}
-//			SAFE_CA_CALL( AudioUnitSetParameter( StreamSplitterUnit, kMatrixMixerParam_Enable, kAudioUnitScope_Output, 0, 0.0, 0 ) );
-		}
-		++DestInputNumber;
 
-#if RADIO_ENABLED
-		if( bNeedRadioFilter )
-		{
-			SAFE_CA_CALL( CreateAndConnectAudioUnit( kAudioUnitType_Effect, 'Rdio', 'Epic', DestNode, DestInputNumber, StreamFormat, StreamFormat, &RadioNode, &RadioUnit ) );
-			SAFE_CA_CALL( AUGraphConnectNodeInput( AudioDevice->GetAudioUnitGraph(), StreamSplitterNode, DestInputNumber, RadioNode, 0 ) );
-			++DestInputNumber;
-		}
+#if CORE_AUDIO_REVERB_ENABLED
+	bNeedReverbFilter = IsReverbApplied();
 #endif
-#if REVERB_ENABLED
-		if( bNeedReverbFilter )
-		{
-			SAFE_CA_CALL( CreateAndConnectAudioUnit( kAudioUnitType_Effect, kAudioUnitSubType_MatrixReverb, kAudioUnitManufacturer_Apple, DestNode, DestInputNumber, StreamFormat, StreamFormat, &ReverbNode, &ReverbUnit ) );
-			SAFE_CA_CALL( AUGraphConnectNodeInput( AudioDevice->GetAudioUnitGraph(), StreamSplitterNode, DestInputNumber, ReverbNode, 0 ) );
-			++DestInputNumber;
-		}
+
+#if CORE_AUDIO_LOWPASS_ENABLED
+	InitLowPassEffect(StreamFormat, HeadNode);
 #endif
 
-		DestNode = StreamSplitterNode;
-		DestInputNumber = 0;
+	// Radio filter will always go first
+	if (bNeedRadioFilter)
+	{
+		InitRadioSourceEffect(StreamFormat, HeadNode);
+	}
+
+	if (bNeedEQFilter)
+	{
+		InitEqSourceEffect(StreamFormat, HeadNode);
+	}
+
+	// Reverb filter always goes last
+	if (bNeedReverbFilter)
+	{
+		InitReverbSourceEffect(StreamFormat, HeadNode);
 	}
 
-	SAFE_CA_CALL( CreateAndConnectAudioUnit( kAudioUnitType_FormatConverter, kAudioUnitSubType_AUConverter, kAudioUnitManufacturer_Apple, DestNode, DestInputNumber, StreamFormat, StreamFormat, &SourceNode, &SourceUnit ) );
+	// Now connect the head node to the final output node
+	// Connect the current head node to the radio node (i.e. source -> radio effect)
+	ErrorStatus = AUGraphConnectNodeInput(AudioDevice->GetAudioUnitGraph(), HeadNode, 0, FinalNode, MixerInputNumber);
+	check(ErrorStatus == noErr);
 
 	if( ErrorStatus == noErr )
 	{
-		AURenderCallbackStruct Input;
-		Input.inputProc = &CoreAudioRenderCallback;
-		Input.inputProcRefCon = this;
-		SAFE_CA_CALL( AudioUnitSetProperty( SourceUnit, kAudioUnitProperty_SetRenderCallback, kAudioUnitScope_Input, 0, &Input, sizeof( Input ) ) );
-
-		// This is split to easily investigate the callstack if the assert happens
 		AUGraph Graph = AudioDevice->GetAudioUnitGraph();
 		check(Graph);
 		SAFE_CA_CALL(AUGraphUpdate( Graph, NULL ));
 
-		GAudioChannels[AudioChannel] = this;
+		AudioDevice->AudioChannels[AudioChannel] = this;
 	}
 	return ErrorStatus == noErr;
 }
@@ -984,40 +874,38 @@ bool FCoreAudioSoundSource::DetachFromAUGraph()
 	Input.inputProcRefCon = NULL;
 	SAFE_CA_CALL( AudioUnitSetProperty( SourceUnit, kAudioUnitProperty_SetRenderCallback, kAudioUnitScope_Input, 0, &Input, sizeof( Input ) ) );
 
-	if( StreamSplitterNode )
-	{
-		SAFE_CA_CALL( AUGraphDisconnectNodeInput( AudioDevice->GetAudioUnitGraph(), StreamSplitterNode, 0 ) );
-	}
-
 // Make sure we still have null nodes
-#if !RADIO_ENABLED
+#if !CORE_AUDIO_RADIO_ENABLED
 	check(RadioNode == 0);
 #endif
 
-#if !REVERB_ENABLED
+#if !CORE_AUDIO_REVERB_ENABLED
 	check(ReverbNode == 0);
 #endif
 
-#if !EQ_ENABLED
+#if !CORE_AUDIO_EQ_ENABLED
 	check(EQNode == 0);
 #endif
 
+#if !CORE_AUDIO_LOWPASS_ENABLED
+	check(LowPassNode == 0);
+#endif
 
-	if( ReverbNode )
+	if (ReverbNode)
 	{
 		SAFE_CA_CALL( AUGraphDisconnectNodeInput( AudioDevice->GetAudioUnitGraph(), ReverbNode, 0 ) );
 	}
-	if( RadioNode )
+	if (RadioNode)
 	{
 		SAFE_CA_CALL( AUGraphDisconnectNodeInput( AudioDevice->GetAudioUnitGraph(), RadioNode, 0 ) );
 	}
-	if( EQNode )
+	if (EQNode)
 	{
 		SAFE_CA_CALL( AUGraphDisconnectNodeInput( AudioDevice->GetAudioUnitGraph(), EQNode, 0 ) );
 	}
-	if( StreamMergerNode )
+	if (LowPassNode)
 	{
-		SAFE_CA_CALL( AUGraphDisconnectNodeInput( AudioDevice->GetAudioUnitGraph(), StreamMergerNode, 0 ) );
+		SAFE_CA_CALL( AUGraphDisconnectNodeInput( AudioDevice->GetAudioUnitGraph(), LowPassNode, 0 ) );
 	}
 
 	if( AudioChannel )
@@ -1034,9 +922,9 @@ bool FCoreAudioSoundSource::DetachFromAUGraph()
 		}
 	}
 
-	if( StreamMergerNode )
+	if (LowPassNode)
 	{
-		SAFE_CA_CALL( AUGraphRemoveNode( AudioDevice->GetAudioUnitGraph(), StreamMergerNode ) );
+		SAFE_CA_CALL( AUGraphRemoveNode( AudioDevice->GetAudioUnitGraph(), LowPassNode ) );
 	}
 	if( EQNode )
 	{
@@ -1050,10 +938,6 @@ bool FCoreAudioSoundSource::DetachFromAUGraph()
 	{
 		SAFE_CA_CALL( AUGraphRemoveNode( AudioDevice->GetAudioUnitGraph(), ReverbNode ) );
 	}
-	if( StreamSplitterNode )
-	{
-		SAFE_CA_CALL( AUGraphRemoveNode( AudioDevice->GetAudioUnitGraph(), StreamSplitterNode ) );
-	}
 	if( AudioChannel )
 	{
 		SAFE_CA_CALL( AUGraphRemoveNode( AudioDevice->GetAudioUnitGraph(), SourceNode ) );
@@ -1064,21 +948,19 @@ bool FCoreAudioSoundSource::DetachFromAUGraph()
 	AudioConverterDispose( CoreAudioConverter );
 	CoreAudioConverter = NULL;
 
-	StreamMergerNode = 0;
-	StreamMergerUnit = NULL;
+	LowPassNode = 0;
+	LowPassUnit = nullptr;
 	EQNode = 0;
-	EQUnit = NULL;
+	EQUnit = nullptr;
 	RadioNode = 0;
-	RadioUnit = NULL;
+	RadioUnit = nullptr;
 	ReverbNode = 0;
-	ReverbUnit = NULL;
-	StreamSplitterNode = 0;
-	StreamSplitterUnit = NULL;
+	ReverbUnit = nullptr;
 	SourceNode = 0;
-	SourceUnit = NULL;
+	SourceUnit = nullptr;
 	MixerInputNumber = -1;
-	
-	GAudioChannels[AudioChannel] = NULL;
+
+	AudioDevice->AudioChannels[AudioChannel] = nullptr;
 	AudioChannel = 0;
 	
 	return true;
diff --git a/Engine/Source/Runtime/Mac/CoreAudio/Public/CoreAudioDevice.h b/Engine/Source/Runtime/Mac/CoreAudio/Public/CoreAudioDevice.h
--- a/Engine/Source/Runtime/Mac/CoreAudio/Public/CoreAudioDevice.h
+++ b/Engine/Source/Runtime/Mac/CoreAudio/Public/CoreAudioDevice.h
@@ -172,7 +172,6 @@ public:
 
 typedef FAsyncTask<class FAsyncRealtimeAudioTaskWorker<FCoreAudioSoundBuffer>> FAsyncRealtimeAudioTask;
 
-
 /**
  * CoreAudio implementation of FSoundSource, the interface used to play, stop and update sources
  */
@@ -294,12 +293,12 @@ protected:
 	AUNode						SourceNode;
 	AudioUnit					SourceUnit;
 
-	AUNode						StreamSplitterNode;
-	AudioUnit					StreamSplitterUnit;
-
 	AUNode						EQNode;
 	AudioUnit					EQUnit;
 
+	AUNode						LowPassNode;
+	AudioUnit					LowPassUnit;
+
 	AUNode						RadioNode;
 	AudioUnit					RadioUnit;
 	bool						bRadioMuted;
@@ -310,9 +309,6 @@ protected:
 
 	bool						bDryMuted;
 
-	AUNode						StreamMergerNode;
-	AudioUnit					StreamMergerUnit;
-
 	int32						AudioChannel;
 	int32						BufferInUse;
 	int32						NumActiveBuffers;
@@ -323,6 +319,12 @@ private:
 
 	void FreeResources();
 
+	void InitSourceUnit(AudioStreamBasicDescription* Format, AUNode& HeadNode);
+	void InitLowPassEffect(AudioStreamBasicDescription* Format, AUNode& HeadNode);
+	void InitRadioSourceEffect(AudioStreamBasicDescription* Format, AUNode& HeadNode);
+	void InitEqSourceEffect(AudioStreamBasicDescription* Format, AUNode& HeadNode);
+	void InitReverbSourceEffect(AudioStreamBasicDescription* Format, AUNode& HeadNode);
+
 	friend class FCoreAudioDevice;
 	friend class FCoreAudioEffectsManager;
 };
@@ -391,6 +393,19 @@ class FCoreAudioDevice : public FAudioDevice
 		return ( ( InputNum << 16 ) | ( OutputNum & 0x0000FFFF ) );
 	}
 
+	int32 FindFreeAudioChannel()
+	{
+		for (int32 Index = 1; Index < CORE_AUDIO_MAX_CHANNELS + 1; Index++)
+		{
+			if (AudioChannels[Index] == nullptr)
+			{
+				return Index;
+			}
+		}
+		return 0;
+	}
+
+
 protected:
 
 	/** Inverse listener transformation, used for spatialization */
@@ -409,11 +424,14 @@ private:
 	AudioStreamBasicDescription	MatrixMixerInputFormat;
 	AudioStreamBasicDescription	MatrixMixerOutputFormat;
 
-	bool						Mixer3DInputStatus[CORE_AUDIO_MAX_MULTICHANNEL_AUDIOCHANNELS];
-	bool						MatrixMixerInputStatus[CORE_AUDIO_MAX_CHANNELS];
+	bool						Mixer3DInputStatus[CORE_AUDIO_MAX_CHANNELS];
+	bool						MatrixMixerInputStatus[CORE_AUDIO_MAX_MULTICHANNEL_AUDIOCHANNELS];
+
+	class FCoreAudioSoundSource* AudioChannels[CORE_AUDIO_MAX_CHANNELS + 1];
 
 	friend class FCoreAudioSoundBuffer;
 	friend class FCoreAudioSoundSource;
+	friend class FCoreAudioEffectsManager;
 };
 
 #endif
diff --git a/Engine/Source/Runtime/Mac/CoreAudio/Public/CoreAudioEffects.h b/Engine/Source/Runtime/Mac/CoreAudio/Public/CoreAudioEffects.h
--- a/Engine/Source/Runtime/Mac/CoreAudio/Public/CoreAudioEffects.h
+++ b/Engine/Source/Runtime/Mac/CoreAudio/Public/CoreAudioEffects.h
@@ -7,9 +7,10 @@
 #ifndef _INC_COREAUDIOEFFECTS
 #define _INC_COREAUDIOEFFECTS
 
-#define REVERB_ENABLED 1
-#define EQ_ENABLED 1
-#define RADIO_ENABLED 1
+#define CORE_AUDIO_LOWPASS_ENABLED 1
+#define CORE_AUDIO_REVERB_ENABLED 1
+#define CORE_AUDIO_EQ_ENABLED 1
+#define CORE_AUDIO_RADIO_ENABLED 1
 
 /** 
  * CoreAudio effects manager
