UE-19303 Fixing crash when pausing/unpausing game in Mac

diff --git a/Engine/Source/Runtime/Mac/CoreAudio/Private/CoreAudioSource.cpp b/Engine/Source/Runtime/Mac/CoreAudio/Private/CoreAudioSource.cpp
--- a/Engine/Source/Runtime/Mac/CoreAudio/Private/CoreAudioSource.cpp
+++ b/Engine/Source/Runtime/Mac/CoreAudio/Private/CoreAudioSource.cpp
@@ -378,13 +378,22 @@ void FCoreAudioSoundSource::Play( void )
 {
 	if( WaveInstance )
 	{
-		if( AttachToAUGraph() )
+		if (!Paused)
 		{
+			if (AttachToAUGraph())
+			{
+				Paused = false;
+				Playing = true;
+
+				// Updates the source which e.g. sets the pitch and volume.
+				Update();
+			}
+		}
+		else
+		{
+			// No need to re-attach the sound to the graph if its just unpausing
 			Paused = false;
 			Playing = true;
-
-			// Updates the source which e.g. sets the pitch and volume.
-			Update();
 		}
 	}
 }
@@ -419,12 +428,8 @@ void FCoreAudioSoundSource::Stop( void )
 void FCoreAudioSoundSource::Pause( void )
 {
 	if( WaveInstance )
-	{	
-		if( Playing && AudioChannel )
-		{
-			DetachFromAUGraph();
-		}
-
+	{
+		// Note: no need to detach from graph when pausing (this is not stopping a sound)
 		Paused = true;
 	}
 }
