UE-19833 Disabling reverb on mac until we can schedule time to refactor/fix mac audio. This should significantly improve performance on Mac.

diff --git a/Engine/Source/Runtime/Mac/CoreAudio/Public/CoreAudioEffects.h b/Engine/Source/Runtime/Mac/CoreAudio/Public/CoreAudioEffects.h
--- a/Engine/Source/Runtime/Mac/CoreAudio/Public/CoreAudioEffects.h
+++ b/Engine/Source/Runtime/Mac/CoreAudio/Public/CoreAudioEffects.h
@@ -8,7 +8,8 @@
 #define _INC_COREAUDIOEFFECTS
 
 #define CORE_AUDIO_LOWPASS_ENABLED 1
-#define CORE_AUDIO_REVERB_ENABLED 1
+// TODO: Disabling reverb until we fix the implementation to use a single AudioUnit instance.
+#define CORE_AUDIO_REVERB_ENABLED 0
 #define CORE_AUDIO_EQ_ENABLED 1
 #define CORE_AUDIO_RADIO_ENABLED 1
 
