UE-19978 Fixing sounds getting stolen on Mac

- reverted the part of CL 2623909 that reduced voice counts on mac.

diff --git a/Engine/Source/Runtime/Mac/CoreAudio/Public/CoreAudioDevice.h b/Engine/Source/Runtime/Mac/CoreAudio/Public/CoreAudioDevice.h
--- a/Engine/Source/Runtime/Mac/CoreAudio/Public/CoreAudioDevice.h
+++ b/Engine/Source/Runtime/Mac/CoreAudio/Public/CoreAudioDevice.h
@@ -29,15 +29,8 @@ DECLARE_LOG_CATEGORY_EXTERN(LogCoreAudio, Log, All);
  * Maximum number of multichannel audio channels - used only for MatrixMixer setup
  */
 
-//#define CORE_AUDIO_MAX_CHANNELS (MAX_AUDIOCHANNELS)
-//#define CORE_AUDIO_MAX_MULTICHANNEL_AUDIOCHANNELS 16
-
-// Note: I'm reducing the core-audio max channels to 32 since PC uses a max of 32
-// TODO: This should be something we define from an engine .ini.
-#define CORE_AUDIO_MAX_CHANNELS (32)
-
-// Note: reducing the max multichannel files to 8 (from 16) for perf reasons
-#define CORE_AUDIO_MAX_MULTICHANNEL_AUDIOCHANNELS 8
+#define CORE_AUDIO_MAX_CHANNELS (MAX_AUDIOCHANNELS)
+#define CORE_AUDIO_MAX_MULTICHANNEL_AUDIOCHANNELS 16
 
 class FCoreAudioDevice;
 class FCoreAudioEffectsManager;
