[NO NON-PII COMMIT MESSAGE]

diff --git a/Engine/Source/Runtime/Mac/CoreAudio/Private/CoreAudioBuffer.cpp b/Engine/Source/Runtime/Mac/CoreAudio/Private/CoreAudioBuffer.cpp
--- a/Engine/Source/Runtime/Mac/CoreAudio/Private/CoreAudioBuffer.cpp
+++ b/Engine/Source/Runtime/Mac/CoreAudio/Private/CoreAudioBuffer.cpp
@@ -18,7 +18,7 @@
  FCoreAudioSoundBuffer.
  ------------------------------------------------------------------------------------*/
 
-static int32 UseResourceTrackingCVar = 0;
+static int32 UseResourceTrackingCVar = 1;
 FAutoConsoleVariableRef CVarUseResourceTracking(
 	TEXT("au.mac.UseResourceTracking"),
 	UseResourceTrackingCVar,
@@ -269,8 +269,6 @@ FCoreAudioSoundBuffer* FCoreAudioSoundBuffer::CreateNativeBuffer( FCoreAudioDevi
 	Buffer->PCMData = Wave->RawPCMData;
 	Buffer->PCMDataSize = Wave->RawPCMDataSize;
 	
-	Wave->RawPCMData = NULL;
-
 	// Keep track of associated resource name.
 	Buffer->InitAudioStreamBasicDescription( kAudioFormatLinearPCM, Wave, true );
 	
@@ -279,6 +277,7 @@ FCoreAudioSoundBuffer* FCoreAudioSoundBuffer::CreateNativeBuffer( FCoreAudioDevi
     
 	if (UseResourceTrackingCVar)
 	{
+		Wave->RawPCMData = NULL;
 		AudioDeviceManager->TrackResource(Wave, Buffer);
 		Wave->RemoveAudioResource();
 	}
