CIS fix (platform-specific typo)

diff --git a/Engine/Source/Runtime/Mac/CoreAudio/Private/CoreAudioBuffer.cpp b/Engine/Source/Runtime/Mac/CoreAudio/Private/CoreAudioBuffer.cpp
--- a/Engine/Source/Runtime/Mac/CoreAudio/Private/CoreAudioBuffer.cpp
+++ b/Engine/Source/Runtime/Mac/CoreAudio/Private/CoreAudioBuffer.cpp
@@ -126,7 +126,7 @@ void FCoreAudioSoundBuffer::InitAudioStreamBasicDescription( UInt32 FormatID, US
 bool FCoreAudioSoundBuffer::ReadCompressedData(uint8* Destination, int32 NumFramesToDecode, bool bLooping)
 {
 	const uint32 kPCMBufferSize = MONO_PCM_BUFFER_SIZE * NumChannels;
-	int32 NumBytesStreamed = (int32)kPCMBUfferSize;
+	int32 NumBytesStreamed = (int32)kPCMBufferSize;
 	if (SoundFormat == SoundFormat_Streaming)
 	{
 		return DecompressionState->StreamCompressedData(Destination, bLooping, kPCMBufferSize, NumBytesStreamed);
